Switchboard: Update references to UE4Editor.

diff --git a/Engine/Plugins/VirtualProduction/Switchboard/Source/Switchboard/switchboard/add_config_dialog.py b/Engine/Plugins/VirtualProduction/Switchboard/Source/Switchboard/switchboard/add_config_dialog.py
--- a/Engine/Plugins/VirtualProduction/Switchboard/Source/Switchboard/switchboard/add_config_dialog.py
+++ b/Engine/Plugins/VirtualProduction/Switchboard/Source/Switchboard/switchboard/add_config_dialog.py
@@ -447,7 +447,7 @@ class AddConfigDialog(QtWidgets.QDialog):
 
         # Not detecting Debug runs since it is not common and would increase detection time
         # At some point this might need UE6Editor added.
-        UEnames = ['UE4Editor', 'UE5Editor']
+        UEnames = ['UE4Editor', 'UnrealEditor']
 
         if sys.platform.startswith('win'):
             for UEname in UEnames:
diff --git a/Engine/Plugins/VirtualProduction/Switchboard/Source/Switchboard/switchboard/devices/ndisplay/ndisplay_monitor.py b/Engine/Plugins/VirtualProduction/Switchboard/Source/Switchboard/switchboard/devices/ndisplay/ndisplay_monitor.py
--- a/Engine/Plugins/VirtualProduction/Switchboard/Source/Switchboard/switchboard/devices/ndisplay/ndisplay_monitor.py
+++ b/Engine/Plugins/VirtualProduction/Switchboard/Source/Switchboard/switchboard/devices/ndisplay/ndisplay_monitor.py
@@ -491,7 +491,7 @@ class nDisplayMonitor(QAbstractTableModel):
 
     @QtCore.Slot()
     def btnFixExeFlags_clicked(self):
-        ''' Tries to force the correct UE4Editor.exe flags '''
+        ''' Tries to force the correct UnrealEditor.exe flags '''
         for devicedata in self.devicedatas.values():
             device = devicedata['device']
             data = devicedata['data']
diff --git a/Engine/Plugins/VirtualProduction/Switchboard/Source/Switchboard/switchboard/devices/ndisplay/plugin_ndisplay.py b/Engine/Plugins/VirtualProduction/Switchboard/Source/Switchboard/switchboard/devices/ndisplay/plugin_ndisplay.py
--- a/Engine/Plugins/VirtualProduction/Switchboard/Source/Switchboard/switchboard/devices/ndisplay/plugin_ndisplay.py
+++ b/Engine/Plugins/VirtualProduction/Switchboard/Source/Switchboard/switchboard/devices/ndisplay/plugin_ndisplay.py
@@ -322,7 +322,7 @@ class DevicenDisplay(DeviceUnreal):
         'executable_filename': Setting(
             attr_name="executable_filename",
             nice_name="nDisplay Executable Filename",
-            value="UE4Editor.exe",
+            value="UnrealEditor.exe",
         ),
         'ndisplay_cmd_args': Setting(
             attr_name="ndisplay_cmd_args",
diff --git a/Engine/Plugins/VirtualProduction/Switchboard/Source/Switchboard/switchboard/devices/unreal/plugin_unreal.py b/Engine/Plugins/VirtualProduction/Switchboard/Source/Switchboard/switchboard/devices/unreal/plugin_unreal.py
--- a/Engine/Plugins/VirtualProduction/Switchboard/Source/Switchboard/switchboard/devices/unreal/plugin_unreal.py
+++ b/Engine/Plugins/VirtualProduction/Switchboard/Source/Switchboard/switchboard/devices/unreal/plugin_unreal.py
@@ -272,10 +272,10 @@ class DeviceUnreal(Device):
             value=0,
             tool_tip="An ID that groups Stage Monitor providers and monitors. Instances with different Session IDs are invisible to each other in Stage Monitor.",
         ),
-        'ue4_exe': Setting(
+        'ue_exe': Setting(
             attr_name="editor_exe",
-            nice_name="UE4 Editor filename",
-            value="UE4Editor.exe",
+            nice_name="Unreal Editor filename",
+            value="UnrealEditor.exe",
         ),
         'max_gpu_count': Setting(
             attr_name="max_gpu_count",
@@ -771,7 +771,7 @@ class DeviceUnreal(Device):
             self.status = DeviceStatus.CLOSED
 
     def fix_exe_flags(self):
-        ''' Tries to force the correct UE4Editor.exe flags '''
+        ''' Tries to force the correct UnrealEditor.exe flags '''
         unreals = self.program_start_queue.running_programs_named('unreal')
 
         if not len(unreals):
