Switchboard: Fix editor device launch. Follow-up to CL 17277528.

diff --git a/Engine/Plugins/VirtualProduction/Switchboard/Source/Switchboard/switchboard/devices/unreal/plugin_unreal.py b/Engine/Plugins/VirtualProduction/Switchboard/Source/Switchboard/switchboard/devices/unreal/plugin_unreal.py
--- a/Engine/Plugins/VirtualProduction/Switchboard/Source/Switchboard/switchboard/devices/unreal/plugin_unreal.py
+++ b/Engine/Plugins/VirtualProduction/Switchboard/Source/Switchboard/switchboard/devices/unreal/plugin_unreal.py
@@ -791,10 +791,9 @@ class DeviceUnreal(Device):
         _, msg = message_protocol.create_minimize_windows_message()
         self.unreal_client.send_message(msg)
 
-
     @property
     def executable_filename(self):
-        return DeviceUnreal.csettings["ue4_exe"].get_value()
+        return DeviceUnreal.csettings['ue_exe'].get_value()
 
     def generate_unreal_exe_path(self):
         return CONFIG.engine_exe_path(CONFIG.ENGINE_DIR.get_value(self.name), self.executable_filename)
