Switchboard: Fix outdated UBT path.

diff --git a/Engine/Plugins/VirtualProduction/Switchboard/Source/Switchboard/switchboard/devices/unreal/plugin_unreal.py b/Engine/Plugins/VirtualProduction/Switchboard/Source/Switchboard/switchboard/devices/unreal/plugin_unreal.py
--- a/Engine/Plugins/VirtualProduction/Switchboard/Source/Switchboard/switchboard/devices/unreal/plugin_unreal.py
+++ b/Engine/Plugins/VirtualProduction/Switchboard/Source/Switchboard/switchboard/devices/unreal/plugin_unreal.py
@@ -731,7 +731,7 @@ class DeviceUnreal(Device):
         program_name = f"build_{program_name_suffix}"
 
         engine_path = CONFIG.ENGINE_DIR.get_value(self.name)
-        ubt_path = os.path.join(engine_path, 'Binaries', 'DotNET', 'UnrealBuildTool')
+        ubt_path = os.path.join(engine_path, 'Binaries', 'DotNET', 'UnrealBuildTool', 'UnrealBuildTool')
 
         puuid, msg = message_protocol.create_start_process_message(
             prog_path=ubt_path, 
