Switchboard: Added support for LegacyFileVersion -8.

#[fyi] Zach.Brockway

diff --git a/Engine/Plugins/VirtualProduction/Switchboard/Source/Switchboard/switchboard/devices/unreal/uassetparser.py b/Engine/Plugins/VirtualProduction/Switchboard/Source/Switchboard/switchboard/devices/unreal/uassetparser.py
--- a/Engine/Plugins/VirtualProduction/Switchboard/Source/Switchboard/switchboard/devices/unreal/uassetparser.py
+++ b/Engine/Plugins/VirtualProduction/Switchboard/Source/Switchboard/switchboard/devices/unreal/uassetparser.py
@@ -277,15 +277,24 @@ class UassetParser(object):
         
         s = UassetSummary()
         
-        s.Tag, s.LegacyFileVersion, s.LegacyUE3Version, s.FileVersionUE4, s.FileVersionLicenseeUE4 \
-        = struct.unpack('<IiiiI', self.f.read(4*5))
-        
+        s.Tag = self.readUInt32()
+
         if s.Tag != 0x9e2a83c1:
             raise ValueError('Not a valid uasset!')
+
+        s.LegacyFileVersion = self.readInt32()
         
-        if s.LegacyFileVersion != -7:
+        if s.LegacyFileVersion not in [-7,-8]:
             raise ValueError('LegacyFileVersion {} is non handled by this parser'.format(s.LegacyFileVersion))
-        
+
+        s.LegacyUE3Version = self.readInt32()
+        s.FileVersionUE4 = self.readInt32()
+
+        if s.LegacyFileVersion == -8:
+            s.FileVersionUE5 = self.readInt32()
+
+        s.FileVersionLicenseeUE4 = self.readUInt32()
+
         s.CustomVersions = self.readTArray(lambda : self.f.read(5*4))        
         s.TotalHeaderSize = self.readInt32()
         s.FolderName = self.readFString()
