Switchboard: use threading.main_thread() object identity test instead of matching by name

threading.main_thread() is available in the threading module as of Python 3.4.

diff --git a/Engine/Plugins/VirtualProduction/Switchboard/Source/Switchboard/switchboard/devices/device_base.py b/Engine/Plugins/VirtualProduction/Switchboard/Source/Switchboard/switchboard/devices/device_base.py
--- a/Engine/Plugins/VirtualProduction/Switchboard/Source/Switchboard/switchboard/devices/device_base.py
+++ b/Engine/Plugins/VirtualProduction/Switchboard/Source/Switchboard/switchboard/devices/device_base.py
@@ -216,7 +216,7 @@ class Device(QtCore.QObject):
     def connect_listener(self):
 
         # Ensure this code is run from the main thread
-        if threading.current_thread().name != 'MainThread':
+        if threading.current_thread() is not threading.main_thread():
             QtCore.QMetaObject.invokeMethod(self, 'connect_listener', QtCore.Qt.QueuedConnection)
             return
 
@@ -228,7 +228,7 @@ class Device(QtCore.QObject):
     def disconnect_listener(self):
 
         # Ensure this code is run from the main thread
-        if threading.current_thread().name != 'MainThread':
+        if threading.current_thread() is not threading.main_thread():
             QtCore.QMetaObject.invokeMethod(self, 'disconnect_listener', QtCore.Qt.QueuedConnection)
             return
 
diff --git a/Engine/Plugins/VirtualProduction/Switchboard/Source/Switchboard/switchboard/devices/unreal/plugin_unreal.py b/Engine/Plugins/VirtualProduction/Switchboard/Source/Switchboard/switchboard/devices/unreal/plugin_unreal.py
--- a/Engine/Plugins/VirtualProduction/Switchboard/Source/Switchboard/switchboard/devices/unreal/plugin_unreal.py
+++ b/Engine/Plugins/VirtualProduction/Switchboard/Source/Switchboard/switchboard/devices/unreal/plugin_unreal.py
@@ -315,7 +315,7 @@ class DeviceUnreal(Device):
     @QtCore.Slot()
     def _queue_notify_redeploy(self):
         # Ensure this code is run from the main thread
-        if threading.current_thread().name != 'MainThread':
+        if threading.current_thread() is not threading.main_thread():
             QtCore.QMetaObject.invokeMethod(self, '_queue_notify_redeploy', QtCore.Qt.QueuedConnection)
             return
 
@@ -553,7 +553,7 @@ class DeviceUnreal(Device):
     def connect_listener(self):
         ''' Connects to the listener '''
         # Ensure this code is run from the main thread
-        if threading.current_thread().name != 'MainThread':
+        if threading.current_thread() is not threading.main_thread():
             QtCore.QMetaObject.invokeMethod(self, 'connect_listener', QtCore.Qt.QueuedConnection)
             return
 
@@ -569,7 +569,7 @@ class DeviceUnreal(Device):
     def disconnect_listener(self):
         ''' Disconnects from the listener '''
         # Ensure this code is run from the main thread
-        if threading.current_thread().name != 'MainThread':
+        if threading.current_thread() is not threading.main_thread():
             QtCore.QMetaObject.invokeMethod(self, 'disconnect_listener', QtCore.Qt.QueuedConnection)
             return
 
diff --git a/Engine/Plugins/VirtualProduction/Switchboard/Source/Switchboard/switchboard/devices/unreal/redeploy_dialog.py b/Engine/Plugins/VirtualProduction/Switchboard/Source/Switchboard/switchboard/devices/unreal/redeploy_dialog.py
--- a/Engine/Plugins/VirtualProduction/Switchboard/Source/Switchboard/switchboard/devices/unreal/redeploy_dialog.py
+++ b/Engine/Plugins/VirtualProduction/Switchboard/Source/Switchboard/switchboard/devices/unreal/redeploy_dialog.py
@@ -98,7 +98,7 @@ class RedeployListenerEndpoint(QtCore.QObject):
     @QtCore.Slot()
     def connect_client(self):
         # Ensure this code is run from the main thread
-        if threading.current_thread().name != 'MainThread':
+        if threading.current_thread() is not threading.main_thread():
             QtCore.QMetaObject.invokeMethod(self, 'connect_client', QtCore.Qt.QueuedConnection)
             return
 
