Switchboard: remove unused ApplicationStatus enum

diff --git a/Engine/Plugins/VirtualProduction/Switchboard/Source/Switchboard/switchboard/switchboard_dialog.py b/Engine/Plugins/VirtualProduction/Switchboard/Source/Switchboard/switchboard/switchboard_dialog.py
--- a/Engine/Plugins/VirtualProduction/Switchboard/Source/Switchboard/switchboard/switchboard_dialog.py
+++ b/Engine/Plugins/VirtualProduction/Switchboard/Source/Switchboard/switchboard/switchboard_dialog.py
@@ -17,7 +17,6 @@ from . import switchboard_utils
 import switchboard.switchboard_widgets as sb_widgets
 
 import datetime
-from enum import IntEnum, unique, auto
 import logging
 import os
 import threading
@@ -29,14 +28,6 @@ RELATIVE_PATH = os.path.dirname(__file__)
 EMPTY_SYNC_ENTRY = "-- None --"
 
 
-@unique
-class ApplicationStatus(IntEnum):
-    DISCONNECTED = auto()
-    CLOSED = auto()
-    SYNCING = auto()
-    OPEN = auto()
-
-
 class SwitchboardDialog(QtCore.QObject):
     def __init__(self):
         super().__init__()
