Switchboard: add an is_disconnected property to Device

In a subsequent change, a CONNECTING status will be added to Device to
represent devices that are in the process of establishing a connection (which
may fail). In many cases, the app will want to treat devices in this state as if
they are disconnected so that UI elements are disabled and requests aren't
attempted using a not-yet-connected device.

This change adds an is_disconnected property that abstracts away the
comparison of the Device's status against DeviceStatus.DISCONNECTED.
When the CONNECTING status is added, it will also be included in the check
done by the is_disconnected property.

diff --git a/Engine/Plugins/VirtualProduction/Switchboard/Source/Switchboard/switchboard/device_list_widget.py b/Engine/Plugins/VirtualProduction/Switchboard/Source/Switchboard/switchboard/device_list_widget.py
--- a/Engine/Plugins/VirtualProduction/Switchboard/Source/Switchboard/switchboard/device_list_widget.py
+++ b/Engine/Plugins/VirtualProduction/Switchboard/Source/Switchboard/switchboard/device_list_widget.py
@@ -155,7 +155,7 @@ class DeviceListWidget(QtWidgets.QListWidget):
         any_connected = False
         any_opened = False
         for device in devices:
-            any_connected |= (device.status != DeviceStatus.DISCONNECTED)
+            any_connected |= (not device.is_disconnected)
             any_opened |= (device.status > DeviceStatus.CLOSED)
 
         header_item = self._header_by_category_name[category_name]
diff --git a/Engine/Plugins/VirtualProduction/Switchboard/Source/Switchboard/switchboard/devices/device_base.py b/Engine/Plugins/VirtualProduction/Switchboard/Source/Switchboard/switchboard/devices/device_base.py
--- a/Engine/Plugins/VirtualProduction/Switchboard/Source/Switchboard/switchboard/devices/device_base.py
+++ b/Engine/Plugins/VirtualProduction/Switchboard/Source/Switchboard/switchboard/devices/device_base.py
@@ -188,6 +188,10 @@ class Device(QtCore.QObject):
         self._status = value
         self.device_qt_handler.signal_device_status_changed.emit(self, previous_status)
 
+    @property
+    def is_disconnected(self):
+        return self.status in {DeviceStatus.DISCONNECTED}
+
     @property
     def project_changelist(self):
         return self._project_changelist
@@ -221,7 +225,7 @@ class Device(QtCore.QObject):
             return
 
         # If the device was disconnected, set to closed
-        if self.status == DeviceStatus.DISCONNECTED:
+        if self.is_disconnected:
             self.status = DeviceStatus.CLOSED
 
     @QtCore.Slot()
@@ -250,7 +254,7 @@ class Device(QtCore.QObject):
             LOGGER.error(f'{self.name}: Incorrect ip address when sending OSC message. {e}')
 
     def record_start(self, slate, take, description):
-        if self.status == DeviceStatus.DISCONNECTED or not self.is_recording_device:
+        if self.is_disconnected or not self.is_recording_device:
             return
 
         self.send_osc_message(osc.RECORD_START, [slate, take, description])
diff --git a/Engine/Plugins/VirtualProduction/Switchboard/Source/Switchboard/switchboard/devices/kipro/plugin_kipro.py b/Engine/Plugins/VirtualProduction/Switchboard/Source/Switchboard/switchboard/devices/kipro/plugin_kipro.py
--- a/Engine/Plugins/VirtualProduction/Switchboard/Source/Switchboard/switchboard/devices/kipro/plugin_kipro.py
+++ b/Engine/Plugins/VirtualProduction/Switchboard/Source/Switchboard/switchboard/devices/kipro/plugin_kipro.py
@@ -82,7 +82,7 @@ class DeviceKiPro(Device):
 
     @unresponsive_kipro
     def record_start(self, slate, take, description):
-        if self.status == DeviceStatus.DISCONNECTED or not self.is_recording_device:
+        if self.is_disconnected or not self.is_recording_device:
             return
 
         self.set_media_state_for_record_play()
@@ -150,7 +150,7 @@ class DeviceKiPro(Device):
 
     @unresponsive_kipro
     def play(self):
-        if self.status == DeviceStatus.DISCONNECTED:
+        if self.is_disconnected:
             return
 
         self.set_media_state_for_record_play()
@@ -160,7 +160,7 @@ class DeviceKiPro(Device):
 
     @unresponsive_kipro
     def stop(self):
-        #if self.status == DeviceStatus.DISCONNECTED:
+        #if self.is_disconnected:
         #    return
 
         self.set_media_state_for_record_play()
@@ -170,7 +170,7 @@ class DeviceKiPro(Device):
 
     @unresponsive_kipro
     def record(self):
-        if self.status == DeviceStatus.DISCONNECTED:
+        if self.is_disconnected:
             return
 
         self.set_media_state_for_record_play()
@@ -180,7 +180,7 @@ class DeviceKiPro(Device):
 
     @unresponsive_kipro
     def fast_forward(self):
-        if self.status == DeviceStatus.DISCONNECTED:
+        if self.is_disconnected:
             return
 
         self.set_media_state_for_record_play()
@@ -190,7 +190,7 @@ class DeviceKiPro(Device):
 
     @unresponsive_kipro
     def fast_reverse(self):
-        if self.status == DeviceStatus.DISCONNECTED:
+        if self.is_disconnected:
             return
 
         self.set_media_state_for_record_play()
diff --git a/Engine/Plugins/VirtualProduction/Switchboard/Source/Switchboard/switchboard/devices/livelinkface/plugin_livelinkface.py b/Engine/Plugins/VirtualProduction/Switchboard/Source/Switchboard/switchboard/devices/livelinkface/plugin_livelinkface.py
--- a/Engine/Plugins/VirtualProduction/Switchboard/Source/Switchboard/switchboard/devices/livelinkface/plugin_livelinkface.py
+++ b/Engine/Plugins/VirtualProduction/Switchboard/Source/Switchboard/switchboard/devices/livelinkface/plugin_livelinkface.py
@@ -46,7 +46,7 @@ class DeviceLiveLinkFace(Device):
 
     def connect_listener(self):
         # If the device was disconnected, set to closed
-        if self.status == DeviceStatus.DISCONNECTED:
+        if self.is_disconnected:
             self.status = DeviceStatus.CLOSED
 
         self.status = DeviceStatus.READY
@@ -111,13 +111,13 @@ class DeviceLiveLinkFace(Device):
 
     def osc_add_send_target_confirm(self):
         # If the mobile is left on and CP is relaunched the normal bootup sequence doesn't happen 
-        if self.status == DeviceStatus.DISCONNECTED:
+        if self.is_disconnected:
             #self.close() # Close the session if it was left open
             self.status = DeviceStatus.OPEN
             self.launch()
 
     def record_start(self, slate, take, description):
-        if self.status == DeviceStatus.DISCONNECTED or not self.is_recording_device:
+        if self.is_disconnected or not self.is_recording_device:
             return
 
         self.send_osc_message(osc.RECORD_START, [slate, take])
diff --git a/Engine/Plugins/VirtualProduction/Switchboard/Source/Switchboard/switchboard/devices/motive/plugin_motive.py b/Engine/Plugins/VirtualProduction/Switchboard/Source/Switchboard/switchboard/devices/motive/plugin_motive.py
--- a/Engine/Plugins/VirtualProduction/Switchboard/Source/Switchboard/switchboard/devices/motive/plugin_motive.py
+++ b/Engine/Plugins/VirtualProduction/Switchboard/Source/Switchboard/switchboard/devices/motive/plugin_motive.py
@@ -127,7 +127,7 @@ class DeviceMotive(Device):
         # because the socket is UDP, we can only be sure that a "connection" has been established
         # when the connect response has beeen received. that's why the connection state is updated
         # here instead of in connect_listener().
-        if self.status == DeviceStatus.DISCONNECTED:
+        if self.is_disconnected:
             self.status = DeviceStatus.CLOSED
 
     def disconnect_listener(self):
@@ -236,7 +236,7 @@ class DeviceMotive(Device):
 
     def record_start(self, slate, take, description):
         """ Called by switchboard_dialog when recording was started, will start recording in Motive. """
-        if self.status == DeviceStatus.DISCONNECTED or not self.trigger_start:
+        if self.is_disconnected or not self.trigger_start:
             return
 
         self.set_slate(slate)
@@ -250,7 +250,7 @@ class DeviceMotive(Device):
 
     def record_stop(self):
         """ Called by switchboard_dialog when recording was stopped, will stop recording in Motive. """
-        if self.status == DeviceStatus.DISCONNECTED or not self.trigger_stop:
+        if self.is_disconnected or not self.trigger_stop:
             return
 
         self.send_request_to_motive('StopRecording')
diff --git a/Engine/Plugins/VirtualProduction/Switchboard/Source/Switchboard/switchboard/devices/shogun/plugin_shogun.py b/Engine/Plugins/VirtualProduction/Switchboard/Source/Switchboard/switchboard/devices/shogun/plugin_shogun.py
--- a/Engine/Plugins/VirtualProduction/Switchboard/Source/Switchboard/switchboard/devices/shogun/plugin_shogun.py
+++ b/Engine/Plugins/VirtualProduction/Switchboard/Source/Switchboard/switchboard/devices/shogun/plugin_shogun.py
@@ -101,7 +101,7 @@ class DeviceShogun(Device):
 
     @unresponsive_shogun
     def record_start(self, slate, take, description):
-        if self.status == DeviceStatus.DISCONNECTED or not self.trigger_start:
+        if self.is_disconnected or not self.trigger_start:
             return
 
         self.set_slate(slate)
@@ -114,7 +114,7 @@ class DeviceShogun(Device):
 
     @unresponsive_shogun
     def record_stop(self):
-        if self.status == DeviceStatus.DISCONNECTED or not self.trigger_stop:
+        if self.is_disconnected or not self.trigger_stop:
             return
 
         result = self.capture_service.stop_capture(0)
diff --git a/Engine/Plugins/VirtualProduction/Switchboard/Source/Switchboard/switchboard/devices/sounddevices/plugin_sounddevices.py b/Engine/Plugins/VirtualProduction/Switchboard/Source/Switchboard/switchboard/devices/sounddevices/plugin_sounddevices.py
--- a/Engine/Plugins/VirtualProduction/Switchboard/Source/Switchboard/switchboard/devices/sounddevices/plugin_sounddevices.py
+++ b/Engine/Plugins/VirtualProduction/Switchboard/Source/Switchboard/switchboard/devices/sounddevices/plugin_sounddevices.py
@@ -65,7 +65,7 @@ class DeviceSoundDevices(Device):
             LOGGER.error(e)
 
     def record_start(self, slate, take, description):
-        if self.status == DeviceStatus.DISCONNECTED or not self.is_recording_device:
+        if self.is_disconnected or not self.is_recording_device:
             return
 
         # Set Slate/Take   
diff --git a/Engine/Plugins/VirtualProduction/Switchboard/Source/Switchboard/switchboard/devices/unreal/plugin_unreal.py b/Engine/Plugins/VirtualProduction/Switchboard/Source/Switchboard/switchboard/devices/unreal/plugin_unreal.py
--- a/Engine/Plugins/VirtualProduction/Switchboard/Source/Switchboard/switchboard/devices/unreal/plugin_unreal.py
+++ b/Engine/Plugins/VirtualProduction/Switchboard/Source/Switchboard/switchboard/devices/unreal/plugin_unreal.py
@@ -435,7 +435,7 @@ class DeviceUnreal(Device):
     def on_build_engine_changed(self, _, build_engine):
         if build_engine:
             self.widget.engine_changelist_label.show()
-            if self.status != DeviceStatus.DISCONNECTED:
+            if not self.is_disconnected:
                 self._request_engine_changelist_number()
         else:
             self.widget.engine_changelist_label.hide()
diff --git a/Engine/Plugins/VirtualProduction/Switchboard/Source/Switchboard/switchboard/switchboard_dialog.py b/Engine/Plugins/VirtualProduction/Switchboard/Source/Switchboard/switchboard/switchboard_dialog.py
--- a/Engine/Plugins/VirtualProduction/Switchboard/Source/Switchboard/switchboard/switchboard_dialog.py
+++ b/Engine/Plugins/VirtualProduction/Switchboard/Source/Switchboard/switchboard/switchboard_dialog.py
@@ -721,7 +721,7 @@ class SwitchboardDialog(QtCore.QObject):
                 continue
 
             # Do not send updates to disconnected devices
-            if device.status == DeviceStatus.DISCONNECTED:
+            if device.is_disconnected:
                 continue
 
             device.set_slate(self._slate)
@@ -771,7 +771,7 @@ class SwitchboardDialog(QtCore.QObject):
         # Tell all devices the new take
         for device in self.device_manager.devices():
             # Do not send updates to disconnected devices
-            if device.status == DeviceStatus.DISCONNECTED:
+            if device.is_disconnected:
                 continue
 
             if device.ip_address in exclude_ip_addresses:
@@ -1006,7 +1006,7 @@ class SwitchboardDialog(QtCore.QObject):
     def update_connect_and_open_button_states(self):
         """ Refresh states of connect-all and start-all buttons. """
         devices = self.device_manager.devices()
-        any_connected = any(device.status != DeviceStatus.DISCONNECTED for device in devices)
+        any_connected = any(not device.is_disconnected for device in devices)
         any_started = any(device.status > DeviceStatus.CLOSED for device in devices)
 
         self.update_connect_all_button_state(any_connected)
@@ -1241,8 +1241,8 @@ class SwitchboardDialog(QtCore.QObject):
             LOGGER.warning(f'{ip_address} is not registered with a device in Switchboard')
             return None
 
-        # Do not recieve osc from disconnected devices
-        if device.status == DeviceStatus.DISCONNECTED:
+        # Do not receive osc from disconnected devices
+        if device.is_disconnected:
             LOGGER.warning(f'{device.name}: is sending OSC commands but is not connected. Ignoring "{command} {value}"')
             return None
 
@@ -1308,7 +1308,7 @@ class SwitchboardDialog(QtCore.QObject):
                 continue
 
             # Do not send updates to disconnected devices
-            if device.status == DeviceStatus.DISCONNECTED:
+            if device.is_disconnected:
                 continue
 
             LOGGER.debug(f'Record Start {device.name}')
@@ -1344,7 +1344,7 @@ class SwitchboardDialog(QtCore.QObject):
                 continue
 
             # Do not send updates to disconnected devices
-            if device.status == DeviceStatus.DISCONNECTED:
+            if device.is_disconnected:
                 continue
 
             device.record_stop()
