Switchboard: Perforce sync overhaul; use of UAT SyncProject replaced with a new Python script called sbl_helper.

Single invocation for combined project and engine sync, correctly reports progress based on synced bytes, avoids constant rebuilds of UAT from dirtied MetaData.cs, easily extensible to support optimizations (UGS filters) and new features (shelved changelists). Handling of updates to Build.version was also made consistent with UGS.

SwitchboardListener: Version bumped to 1.5.0. Cleaned up some optional message field handling. Fixed a race condition where some stdout data could be missed. Also changed stdout packets to base64 encoding to support binary data via stdout (didn't end up using this yet, but seems worth supporting).

diff --git a/Engine/Plugins/VirtualProduction/Switchboard/Source/Switchboard/sbl_helper.py b/Engine/Plugins/VirtualProduction/Switchboard/Source/Switchboard/sbl_helper.py
new file mode 100644
--- /dev/null
+++ b/Engine/Plugins/VirtualProduction/Switchboard/Source/Switchboard/sbl_helper.py
@@ -0,0 +1,863 @@
+# Copyright Epic Games, Inc. All Rights Reserved.
+
+from __future__ import annotations
+
+import argparse
+import copy
+import io
+import json
+import logging
+import marshal
+import os
+import pathlib
+import subprocess
+import sys
+import threading
+import tempfile
+import traceback
+from queue import Empty, Queue
+from typing import Callable, Dict, List, Optional, Sequence, Tuple, Union
+
+
+MAX_PARALLEL_WORKERS = 8 # Consistent with UGS, ushell, etc.
+
+SCRIPT_NAME = os.path.basename(__file__)
+__version__ = '1.0.0'
+
+
+def p4_variables() -> Dict[str, str]:
+    ''' Returns Perforce variable (e.g. `P4CONFIG`, `p4 set`) key/value pairs. '''
+    output = subprocess.check_output(['p4', 'set', '-q']).decode()
+    variables: Dict[str, str] = {}
+    for line in output.splitlines():
+        var_name, _, var_value = line.partition('=')
+        variables[var_name] = var_value
+    return variables
+
+
+def p4(
+    cmd: str,
+    cmd_opts: List[str],
+    global_opts: Optional[List[str]] = None,
+    *,
+    user: Optional[str] = None,
+    client: Optional[str] = None,
+) -> subprocess.Popen:
+    ''' Runs a p4 command and includes some common/required switches. '''
+    args = ['p4', f'-zprog={SCRIPT_NAME}', f'-zversion={__version__}', '-ztag', '-G', '-Qutf8']
+
+    if user is not None:
+        args.extend(['-u', user])
+    if client is not None:
+        args.extend(['-c', client])
+
+    args.extend(global_opts or [])
+
+    args.append(cmd)
+    args.extend(cmd_opts)
+
+    logging.debug(f'p4(): invoking subprocess: args={args}')
+
+    # stdin=subprocess.DEVNULL is required when launched via SwitchboardListener; otherwise this call tries
+    # to make the non-existent stdin inheritable, and raises `OSError: [WinError 6] The handle is invalid`
+    return subprocess.Popen(args, stdout=subprocess.PIPE, stdin=subprocess.DEVNULL)
+
+
+def p4_get_records(
+    cmd: str,
+    opts: List[str],
+    global_opts: Optional[List[str]] = None,
+    *,
+    include_info: bool = False,
+    include_error: bool = False,
+    user: Optional[str] = None,
+    client: Optional[str] = None,
+) -> List[dict]:
+    results: List[dict] = []
+    proc = p4(cmd, opts, global_opts=global_opts, user=user, client=client)
+    while True:
+        try:
+            record = marshal.load(proc.stdout)
+        except EOFError:
+            break
+
+        # Possible codes: stat, info, error, text (only `p4 print`)
+        include_record = False
+        if (record[b'code'] == b'stat') or (record[b'code'] == b'text'):
+            include_record = True
+        elif include_info and (record[b'code'] == b'info'):
+            include_record = True
+        elif include_error and (record[b'code'] == b'error'):
+            include_record = True
+
+        if include_record:
+            results.append(record)
+
+    proc.stdout.close()
+    return results
+
+
+def p4_changes(
+    pathspecs: List[str],
+    *,
+    limit: int = 10,
+    user: Optional[str] = None,
+    client: Optional[str] = None,
+) -> List[dict]:
+    limit = min(limit, 100)
+    opts = ['-s', 'submitted', '-t', '-L', f'-m{limit}', *pathspecs]
+    return p4_get_records('changes', opts, user=user, client=client)
+
+
+def p4_have(
+    pathspecs: List[str],
+    *,
+    user: Optional[str] = None,
+    client: Optional[str] = None,
+) -> List[Optional[int]]:
+    ''' Returns a list of revision numbers, or None if the file doesn't exist on the client. '''
+    result: List[Optional[int]] = []
+    records = p4_get_records('have', pathspecs, include_error=True, user=user, client=client)
+    for record in records:
+        if record[b'code'] == b'stat':
+            result.append(int(record[b'haveRev']))
+        else:
+            result.append(None)
+    return result
+
+
+def p4_print(
+    pathspecs: List[str],
+    *,
+    user: Optional[str] = None,
+    client: Optional[str] = None,
+) -> List[Tuple[Optional[str], dict]]:
+    ''' Returns a list of (fileContents, statRecord).
+        If the file does not exist, or is deleted at the specified revision, fileContents will be None.
+    '''
+    results: List[Tuple[Optional[str], dict]] = []
+    records = p4_get_records('print', pathspecs, include_info=True, include_error=True, user=user, client=client)
+
+    # For each file, a stat record is returned, followed (if present) by a text record.
+    last_stat_record = {}
+    for record in records:
+        if record[b'code'] == b'stat':
+            if b'delete' in record[b'action']: # 'deleted' or 'move/delete'
+                results.append((None, record))
+            else:
+                last_stat_record = record
+            continue
+        elif record[b'code'] == b'text':
+            results.append((record[b'data'].decode(), last_stat_record))
+            continue
+        elif record[b'code'] == b'error':
+            results.append((None, record))
+            continue
+
+        logging.warning(f'p4_print(): Unhandled record: {record}')
+
+    return results
+
+
+class SyncFile:
+    ''' Represents a single file pending sync. '''
+    def __init__(self, depot_path: str, client_path: Optional[str], rev: int, size: int, action: str):
+        self.depot_path = depot_path
+        self.client_path = client_path
+        self.rev = rev
+        self.size = size
+        self.action = action
+
+        self.clobber = False # Used during noclobber resolution, controls whether to force sync in a second pass.
+
+    @property
+    def spec(self) -> str:
+        ''' Returns the full path + rev spec (compatible with `p4 sync -L`). '''
+        return f'{self.depot_path}#{self.rev}'
+
+    def __str__(self) -> str:
+        return f'{self.spec} ({friendly_bytes(self.size)}) ({self.action})'
+
+
+class SyncWorkload:
+    ''' Represents a sync operation, or e.g. a parallel worker's subset thereof. '''
+    def __init__(self):
+        self.sync_files: Dict[str, SyncFile] = {}
+        self.workload_size = 0
+        self.client_path_to_depot_path: Dict[str, str] = {}
+
+    def add_or_replace(self, new_file: SyncFile):
+        if new_file.depot_path in self.sync_files:
+            self.remove(new_file.depot_path)
+        self.sync_files[new_file.depot_path] = new_file
+        self.workload_size += new_file.size
+        if new_file.client_path:
+            self.client_path_to_depot_path[new_file.client_path] = new_file.depot_path
+
+    def add(self, new_file: SyncFile):
+        assert new_file.depot_path not in self.sync_files
+        self.add_or_replace(new_file)
+
+    def remove(self, depot_path: str):
+        remove_file = self.sync_files[depot_path]
+        self.workload_size -= remove_file.size
+        del self.sync_files[depot_path]
+        if remove_file.client_path is not None:
+            if remove_file.client_path in self.client_path_to_depot_path:
+                del self.client_path_to_depot_path[remove_file.client_path]
+
+    def get_by_client_path(self, client_path: str) -> Optional[SyncFile]:
+        if client_path in self.client_path_to_depot_path:
+            depot_path = self.client_path_to_depot_path[client_path]
+            return self.sync_files.get(depot_path)
+
+
+def p4_sync_worker(
+    sync_workload: SyncWorkload,
+    *,
+    completion_queue: Optional[Queue[SyncFile]] = None,
+    clobber_queue: Optional[Queue[SyncFile]] = None,
+    dry_run: bool = False,
+    force: bool = False,
+    user: Optional[str] = None,
+    client: Optional[str] = None,
+):
+    ''' Performs a sync (subset), and reports file sync completions via `completion_queue`. '''
+    gopts = []
+    opts = ['--parallel=0']
+
+    # p4 2020.1 relnotes: "'p4 sync -L' now allows files to be specified with revision specifiers of #0."
+    workload_has_deletes = any(x.rev == 0 for x in sync_workload.sync_files.values())
+    if not workload_has_deletes:
+        opts.append('-L')
+
+    if dry_run:
+        opts.append('-n')
+
+    if force:
+        opts.append('-f')
+
+    argfile = tempfile.NamedTemporaryFile(prefix=f'{SCRIPT_NAME}-', suffix='.txt', mode='wt', delete=False)
+    gopts.extend(['-x', argfile.name])
+
+    for item in sync_workload.sync_files.values():
+        argfile.write(f'{item.spec}\n')
+    argfile.close() # Otherwise the p4 process can't access the file.
+
+    proc = p4('sync', opts, gopts, user=user, client=client)
+
+    remaining_work = copy.deepcopy(sync_workload)
+
+    def complete_item(depot_path: str):
+        completed_file = remaining_work.sync_files[depot_path]
+        remaining_work.remove(depot_path)
+        if completion_queue is not None:
+            completion_queue.put(completed_file)
+
+    CLOBBER_PREFIX = b"Can't clobber writable file "
+
+    while True:
+        try:
+            record = marshal.load(proc.stdout)
+        except EOFError:
+            break
+
+        if record[b'code'] == b'stat':
+            depot_path = record[b'depotFile'].decode()
+            complete_item(depot_path)
+            continue
+        elif record[b'code'] == b'error':
+            data: bytes = record[b'data']
+
+            # Detect file sync cancellation caused by client being noclobber and the file being writable.
+            # Note that this only happens AFTER a typical completion 'stat' record was also emitted.
+            if data.startswith(CLOBBER_PREFIX):
+                clobber_client_path = data[len(CLOBBER_PREFIX):-1].decode() # Also trim \n
+                clobber_file = sync_workload.get_by_client_path(clobber_client_path)
+                if clobber_file:
+                    if clobber_queue:
+                        clobber_queue.put(copy.deepcopy(clobber_file))
+                    else:
+                        # Caller isn't gathering the list of clobber files; just emit the error message.
+                        logging.error(record[b'data'].decode())
+                else:
+                    # Shouldn't happen.
+                    logging.error(f'Unable to map clobber file to depot path: {clobber_client_path}')
+                continue
+
+        # NB: Not currently handling '... up-to-date' here; should have been excluded in preview.
+
+        logging.warning(f'p4_sync_worker(): Unhandled record: {record}')
+
+    proc.stdout.close()
+    os.unlink(argfile.name) # We can clean up the temporary `-x` file now.
+
+    # This shouldn't happen, but if something goes wrong it prevents a deadlock.
+    if len(remaining_work.sync_files):
+        for depot_path, item in remaining_work.sync_files.copy().items():
+            logging.warning(f'p4_sync_worker(): Flushing item not explicitly completed: {item}')
+            complete_item(depot_path)
+
+
+def p4_get_preview_sync_files(
+    input_specs: List[str],
+    *,
+    user: Optional[str] = None,
+    client: Optional[str] = None,
+) -> SyncWorkload:
+    ''' Runs a preview sync to generate a list of files and revisions. '''
+    result_workload = SyncWorkload()
+    records = p4_get_records('sync', ['-n', *input_specs],
+        include_info=True, include_error=True, user=user, client=client)
+
+    for record in records:
+        if record[b'code'] == b'stat':
+            depot_path = record[b'depotFile'].decode()
+            action = record[b'action'].decode()
+            client_path = record[b'clientFile'].decode() if b'clientFile' in record else None
+
+            if action == 'deleted':
+                # Delete records report the 'rev' PRIOR to the delete, and have no 'fileFize'.
+                rev = 0
+                size = 0
+            else:
+                rev = int(record[b'rev'])
+                size = int(record[b'fileSize'])
+
+            result_workload.add_or_replace(SyncFile(depot_path, client_path, rev, size, action))
+            continue
+        elif record[b'code'] == b'info':
+            msg = f'{record[b"data"].decode()}'
+            if record[b'level'] > 0:
+                logging.warning(msg)
+            else:
+                logging.info(msg)
+            continue
+        elif record[b'code'] == b'error':
+            if record[b'data'].endswith(b'up-to-date.\n'):
+                logging.info(record[b'data'].decode().rstrip())
+                continue
+
+        logging.warning(f'p4_get_preview_sync_files(): Unhandled record: {record}')
+
+    return result_workload
+
+
+def friendly_bytes(size: float) -> str:
+    suffixes = ['bytes', 'KiB', 'MiB', 'GiB', 'TiB']
+    mag = 0
+    while size >= 1000:
+        size /= 1024
+        mag += 1
+    return f'{size:.{2 if mag>0 else 0}f} {suffixes[mag]}'
+
+
+def p4_sync(
+    input_specs: List[str],
+    num_workers: int = MAX_PARALLEL_WORKERS,
+    *,
+    sync_filter_fn: Optional[Callable[[SyncWorkload], None]] = None,
+    clobber_filter_fn: Optional[Callable[[SyncWorkload], None]] = None,
+    progress_fn: Optional[Callable[[int, int, int, int], None]] = None,
+    dry_run: bool = False,
+    user: Optional[str] = None,
+    client: Optional[str] = None,
+    write_workload_file: Optional[io.TextIOWrapper] = None
+) -> Optional[int]:
+    # First, gather complete list of files/revisions that need to be synced.
+    total_work = p4_get_preview_sync_files(input_specs, user=user, client=client)
+
+    if sync_filter_fn:
+        sync_filter_fn(total_work)
+
+    if len(total_work.sync_files) < 1:
+        logging.info('Nothing to do')
+        return None
+
+    # User may have requested to write the pending sync files and revisions to file and exit early.
+    if write_workload_file is not None:
+        for item in total_work.sync_files.values():
+            write_workload_file.write(f'{item.spec}\n')
+        return 0
+
+    # Separate deletes and run them on their own worker process. (We no longer refer to `total_work` after this.)
+    # (p4 2020.1 relnotes: "'p4 sync -L' now allows files to be specified with revision specifiers of #0.")
+    nondelete_work = SyncWorkload()
+    delete_work = SyncWorkload()
+    for item in total_work.sync_files.copy().values():
+        if item.rev == 0:
+            delete_work.add(item)
+        else:
+            nondelete_work.add(item)
+
+    # First pass: Attempt to run all deletes (in the main thread).
+    total_delete_files = len(delete_work.sync_files)
+    if total_delete_files > 0:
+        logging.info(f'Removing {total_delete_files} files from workspace...')
+
+        clobber_queue: Optional[Queue[SyncFile]] = None
+        if clobber_filter_fn:
+            clobber_queue = Queue()
+
+        p4_sync_worker(delete_work, clobber_queue=clobber_queue, dry_run=dry_run, user=user, client=client)
+
+        # First pass clobber retry: Force sync any deletes where we want to override noclobber.
+        if clobber_queue and (clobber_queue.qsize() > 0):
+            clobber_work = SyncWorkload()
+            while clobber_queue.qsize() > 0:
+                clobber_work.add(clobber_queue.get())
+
+            clobber_filter_fn(clobber_work)
+
+            for depot_path, file in clobber_work.sync_files.copy().items():
+                if not file.clobber:
+                    logging.error(f'NOT clobbering writable file: {file.client_path or file.depot_path}')
+                    clobber_work.remove(depot_path)
+                else:
+                    logging.warning(f'Forcing clobber of writable file: {file.client_path}')
+
+            num_files_to_clobber = len(clobber_work.sync_files)
+            if num_files_to_clobber:
+                logging.info(f'Clobbering {num_files_to_clobber} files')
+                p4_sync_worker(clobber_work, force=True, dry_run=dry_run, user=user, client=client)
+
+        logging.info('Done removing files from workspace')
+
+    # Second pass: Dispatch all non-delete sync actions to worker threads.
+    # FIXME: This block is almost the same as above, but we don't use multiple workers above.
+    clobber_queue: Optional[Queue[SyncFile]] = None
+    if clobber_filter_fn:
+        clobber_queue = Queue()
+
+    dispatch_sync_workers(
+        nondelete_work,
+        num_workers,
+        dry_run=dry_run,
+        clobber_queue=clobber_queue,
+        progress_fn=progress_fn,
+        user=user,
+        client=client)
+
+    # Second pass clobber retry: Force sync where we want to override noclobber.
+    if clobber_queue and (clobber_queue.qsize() > 0):
+        clobber_work = SyncWorkload()
+        while clobber_queue.qsize() > 0:
+            clobber_work.add(clobber_queue.get())
+
+        clobber_filter_fn(clobber_work)
+
+        for depot_path, file in clobber_work.sync_files.copy().items():
+            if not file.clobber:
+                logging.error(f'NOT clobbering writable file: {file.client_path or file.depot_path}')
+                clobber_work.remove(depot_path)
+            else:
+                logging.warning(f'Forcing clobber of writable file: {file.client_path}')
+
+        num_files_to_clobber = len(clobber_work.sync_files)
+        if num_files_to_clobber:
+            logging.info(f'Clobbering {num_files_to_clobber} files')
+            dispatch_sync_workers(
+                clobber_work,
+                num_workers,
+                dry_run=dry_run,
+                force=True,
+                progress_fn=progress_fn,
+                user=user,
+                client=client)
+
+    logging.info('Sync complete')
+    return 0
+
+
+def dispatch_sync_workers(
+    work: SyncWorkload,
+    num_workers: int,
+    *,
+    dry_run: bool = False,
+    force: bool = False,
+    clobber_queue: Optional[Queue[SyncFile]] = None,
+    progress_fn: Optional[Callable[[int, int, int, int], None]] = None,
+    user: Optional[str] = None,
+    client: Optional[str] = None,
+) -> int:
+    ''' Divides `work` evenly among up to `num_workers` sync subprocesses. '''
+    # Display summary of total sync workload.
+    total_sync_bytes = work.workload_size
+    total_sync_files = len(work.sync_files)
+    logging.info(f'Syncing {friendly_bytes(total_sync_bytes)} in {total_sync_files:,} files')
+
+    # Upper bound is least of 1) CPU cores minus one, 2) number of files or 3) hardcoded max constant.
+    num_workers = min(filter(None, [num_workers, (os.cpu_count() or 1) - 1, total_sync_files, MAX_PARALLEL_WORKERS]))
+
+    # Distribute balanced loads to workers.
+    worker_loads = [SyncWorkload() for _ in range(num_workers)]
+    work.sync_files = dict(sorted(work.sync_files.items(), key=lambda x: x[1].size, reverse=True))
+    for sync_item in work.sync_files.values():
+        lightest_load = min(worker_loads, key=lambda x: (x.workload_size, len(x.sync_files)))
+        lightest_load.add_or_replace(sync_item)
+
+    logging.info(f'Using {num_workers} parallel workers')
+
+    # Used by workers to notify main thread when files finish syncing.
+    completion_queue: Queue[SyncFile] = Queue()
+
+    # Summarize work distribution and start workers.
+    workers: List[threading.Thread] = []
+    for i in range(num_workers):
+        worker_name = f'Worker {i}'
+        load = worker_loads[i]
+        if num_workers > 1:
+            logging.info(f'    {worker_name} - {friendly_bytes(load.workload_size)} in {len(load.sync_files):,} files')
+
+        worker = threading.Thread(
+            name=worker_name,
+            target=p4_sync_worker,
+            args=(load,),
+            kwargs={'completion_queue': completion_queue, 'clobber_queue': clobber_queue,
+                'force': force, 'dry_run': dry_run, 'user': user, 'client': client}
+        )
+        workers.append(worker)
+        worker.start()
+
+    # Wait for workers to signal completed file syncs, and display progress updates.
+    remaining_load = copy.deepcopy(work) # Items are removed from this copy as they are completed.
+    remaining_num_files = len(remaining_load.sync_files)
+    while remaining_num_files:
+        try:
+            completed_file = completion_queue.get(timeout=1.0)
+        except Empty:
+            if any(worker.is_alive() for worker in workers):
+                continue # Workers still running; wait patiently.
+            else:
+                logging.error(f'{remaining_num_files} unsynced files remaining, but all workers have exited')
+                return 1 # Shouldn't happen.
+
+        remaining_load.remove(completed_file.depot_path)
+        remaining_num_files = len(remaining_load.sync_files)
+
+        if progress_fn is not None:
+            completed_bytes = total_sync_bytes - remaining_load.workload_size
+            completed_files = total_sync_files - remaining_num_files
+            progress_fn(completed_bytes, total_sync_bytes, completed_files, total_sync_files)
+
+    for worker in workers:
+        worker.join()
+
+    return 0
+
+
+def path_is_relative_to(root: Union[str, os.PathLike[str]], other: Union[str, os.PathLike[str]]) -> bool:
+    try:
+        pathlib.PurePath(other).relative_to(root)
+        return True
+    except ValueError:
+        return False
+
+
+class SbListenerHelper:
+    def __init__(self):
+        self.parser = self.build_parser()
+
+        self.uproj_path: Optional[pathlib.Path] = None
+        self.engine_dir: Optional[pathlib.Path] = None
+        self.sync_engine_cl: Optional[int] = None
+        self.sync_project_cl: Optional[int] = None
+        self.clobber_engine: bool = False
+        self.clobber_project: bool = False
+        self.generate_after_sync: bool = False
+
+        self.dry_run: bool = False
+        self.p4user: Optional[str] = None
+        self.p4client: Optional[str] = None
+
+
+    @staticmethod
+    def build_parser():
+        parser = argparse.ArgumentParser()
+
+        # Global options
+        parser.add_argument('--log-level', default='INFO', help='DEBUG, INFO (default), WARNING, ERROR, CRITICAL')
+
+        subparsers = parser.add_subparsers(dest='action')
+        subparsers.required = True
+
+        # Sync action
+        sync_parser = subparsers.add_parser('sync')
+
+        sync_parser.add_argument('pathspecs', nargs='*', metavar='pathSpec[revSpec]',
+            help='File/path specifier, supporting * or ... wildcards [and optional # or @ revision specifier suffix]')
+
+        sync_parser.add_argument('--project', type=pathlib.Path, metavar='UPROJECT', help='Path to .uproject file')
+        sync_parser.add_argument('--engine-cl', type=int, help='Changelist to sync engine to (requires UPROJECT)')
+        sync_parser.add_argument('--project-cl', type=int, help='Changelist to sync project to (requires UPROJECT)')
+        sync_parser.add_argument('--generate', action='store_true', help='Run GenerateProjectFiles after syncing')
+        sync_parser.add_argument('--clobber-engine', action='store_true', help='Override noclobber for engine files')
+        sync_parser.add_argument('--clobber-project', action='store_true', help='Override noclobber for project files')
+
+        sync_parser.add_argument('-n', '--dry-run', action='store_true', help='Preview sync; do not update files')
+        sync_parser.add_argument('-u', '--p4user', type=str, help='Override Perforce P4USER')
+        sync_parser.add_argument('-c', '--p4client', type=str, help='Override Perforce P4CLIENT')
+
+        sync_parser.add_argument('--dump-sync', type=argparse.FileType('wt'), metavar='SYNC_WORKLOAD_TEXT_FILE',
+            help='Write the complete list of depot paths and revisions that would be synced, but exit without syncing')
+
+        return parser
+
+
+    def run(self, args: Optional[Sequence[str]] = None) -> int:
+        options = self.parser.parse_args(args)
+
+        logging.basicConfig(
+            level=options.log_level,
+            format='%(relativeCreated)06dms %(levelname)-8s %(message)s'
+        )
+
+        if options.action == 'sync':
+            return self.run_sync(options)
+
+        assert False
+
+
+    def run_sync(self, options: argparse.Namespace) -> int:
+        pathspecs: List[str] = [*options.pathspecs]
+
+        self.check_sync_options(options)
+
+        if self.sync_engine_cl:
+            pathspecs.append(f'{self.engine_dir.parent / "*"}@{self.sync_engine_cl}')
+            pathspecs.append(f'{self.engine_dir / "..."}@{self.sync_engine_cl}')
+
+        if self.sync_project_cl:
+            pathspecs.append(f'{self.uproj_path.parent / "..."}@{self.sync_project_cl}')
+
+        if len(pathspecs) < 1:
+            self.parser.error('Nothing specified to sync')
+        else:
+            logging.info('Scheduled for sync:')
+            for pathspec in pathspecs:
+                logging.info(f' - {pathspec}')
+
+        buildver_info = None
+        if self.sync_engine_cl:
+            buildver_info = self.get_buildver_info()
+            if buildver_info['have_rev'] is not None:
+                pathspecs.append(f"{buildver_info['local_path']}#0")
+
+        def sync_filter(workload: SyncWorkload):
+            # Don't add Build.version to the workspace, but do allow it to be removed.
+            try:
+                if (buildver_info['depot_path'] is not None) and (workload.sync_files[buildver_info['depot_path']].rev != 0):
+                    workload.remove(buildver_info['depot_path'])
+            except KeyError:
+                pass
+
+        def clobber_filter(clobber: SyncWorkload):
+            clobber.sync_files = dict(sorted(clobber.sync_files.items()))
+
+            for depot_path, clobber_file in clobber.sync_files.items():
+                if clobber_file.client_path is not None:
+                    if self.clobber_engine and self.engine_dir:
+                        if path_is_relative_to(self.engine_dir, clobber_file.client_path):
+                            clobber_file.clobber = True
+
+                    if self.clobber_project and self.uproj_path:
+                        if path_is_relative_to(self.uproj_path.parent, clobber_file.client_path):
+                            clobber_file.clobber = True
+                else:
+                    # Shouldn't happen.
+                    logging.error(f'clobber_filter(): Missing client path for clobber file: {depot_path}')
+                    continue
+
+        sync_result = p4_sync(
+            pathspecs,
+            sync_filter_fn=sync_filter,
+            clobber_filter_fn=clobber_filter,
+            progress_fn=self.on_sync_progress,
+            dry_run=self.dry_run,
+            user=self.p4user,
+            client=self.p4client,
+            write_workload_file=options.dump_sync,
+        )
+
+        if sync_result is None:
+            return 0 # Nothing was synced
+        elif self.dry_run or options.dump_sync:
+            return sync_result # Nothing more to do
+
+        # Write updated Build.version.
+        if self.sync_engine_cl:
+            if buildver_info['updated_text'] is not None:
+                logging.info('Updating Build.version')
+                with open(buildver_info['local_path'], 'wt') as buildver_file:
+                    buildver_file.write(buildver_info['updated_text'])
+
+        # Generate project files.
+        if self.generate_after_sync:
+            logging.info('Generating project files...')
+            gpf_result = self.generate_project_files()
+            if gpf_result == 0:
+                logging.info('Done generating project files')
+            else:
+                logging.error(f'GenerateProjectFiles failed with code {gpf_result}')
+                return gpf_result
+
+        return sync_result
+
+
+    def get_buildver_info(self):
+        ''' Read `Build.version` from depot and return its info, along with changes that reflect local build. '''
+        logging.debug('get_buildver_info(): start')
+
+        assert self.engine_dir
+
+        epicint_local_path = os.path.join(self.engine_dir, 'Restricted', 'NotForLicensees', 'Build', 'EpicInternal.txt')
+        buildver_local_path = os.path.join(self.engine_dir, 'Build', 'Build.version')
+
+        buildver_have_rev = p4_have([buildver_local_path], user=self.p4user, client=self.p4client)[0]
+
+        prints = p4_print([f'{path}@{self.sync_engine_cl}' for path in (epicint_local_path, buildver_local_path)],
+            user=self.p4user, client=self.p4client)
+        epicint_depot_text, _ = prints[0]
+        buildver_depot_text, buildver_meta = prints[1]
+
+        updated_buildver_text: Optional[str] = None
+        buildver_depot_path: Optional[str] = None
+
+        if buildver_depot_text is not None:
+            # Code changelist determination compatible with precompiled binaries.
+            code_exts = ['.cs', '.h', '.cpp', '.usf', '.ush', '.uproject', '.uplugin']
+            engine_code_paths = [f'{self.engine_dir}/.../*{ext}@<={self.sync_engine_cl}' for ext in code_exts]
+            engine_code_changes = p4_changes(engine_code_paths, limit=1, user=self.p4user, client=self.p4client)
+            if len(engine_code_changes) > 0:
+                code_cl = max(int(x[b'change']) for x in engine_code_changes)
+            else:
+                logging.warning('Unable to determine last code changelist')
+                code_cl = self.sync_engine_cl
+
+            # Generate and return the updated Build.version contents.
+            try:
+                build_ver: Dict[str, Union[int, str]] = json.loads(buildver_depot_text)
+            except json.JSONDecodeError:
+                logging.error('Unable to parse depot Build.version JSON', exc_info=sys.exc_info())
+                return
+
+            buildver_depot_path = buildver_meta[b'depotFile'].decode()
+            branch_name = buildver_depot_path.replace('/Engine/Build/Build.version', '')
+            updates = {
+                'Changelist': self.sync_engine_cl,
+                'IsLicenseeVersion': 1 if epicint_depot_text is None else 0,
+                'BranchName': branch_name.replace('/', '+'),
+            }
+
+			# Don't overwrite the compatible changelist if we're in a hotfix release
+            if (build_ver['CompatibleChangelist'] == 0) or (build_ver['IsLicenseeVersion'] != updates['IsLicenseeVersion']):
+                updates['CompatibleChangelist'] = code_cl
+
+            logging.debug(f"get_buildver_info(): {updates}")
+
+            build_ver.update(updates)
+            updated_buildver_text = json.dumps(build_ver, indent=4)
+
+        else:
+            logging.error(f'Unable to p4_print() Build.version ({buildver_meta})')
+
+        result = {}
+        result['updated_text'] = updated_buildver_text
+        result['depot_text'] = buildver_depot_text
+        result['depot_path'] = buildver_depot_path
+        result['local_path'] = buildver_local_path
+        result['have_rev'] = buildver_have_rev
+        return result
+
+
+    def on_sync_progress(self, completed_bytes: int, total_bytes: int, completed_files: int, total_files: int):
+        progress_pct = completed_bytes / total_bytes if total_bytes > 0 else completed_files / total_files
+        bytes_str = f'{friendly_bytes(completed_bytes)} / {friendly_bytes(total_bytes)}'
+        files_str = f'{completed_files:,} / {total_files:,} files'
+        progress_str = f'Progress: {progress_pct:.2%} ({bytes_str}; {files_str})'
+
+        if sys.stdout.isatty():
+            # Continuously erase and rewrite interactive progress status line.
+            print(f'\r{" "*70}\r{progress_str}\r', end='' if progress_pct != 1.0 else None, flush=True)
+        else:
+            logging.info(progress_str)
+
+
+    def check_sync_options(self, options: argparse.Namespace):
+        ''' Parse `sync` options, caching their values, and ensuring consistency between related options. '''
+        self.dry_run = options.dry_run
+        self.p4user = options.p4user
+        self.p4client = options.p4client
+
+        consistency_opts = ('project', 'engine_cl', 'project_cl', 'generate', 'clobber_engine', 'clobber_project')
+        if any(getattr(options, opt) is not None for opt in consistency_opts):
+            if options.project is None:
+                self.parser.error('--[engine|project]-cl, --clobber-[engine|project], --generate require --project')
+            elif options.engine_cl is not None and options.engine_cl <= 0:
+                self.parser.error(f'Invalid --engine-cl {options.engine_cl}')
+            elif options.project_cl is not None and options.project_cl <= 0:
+                self.parser.error(f'Invalid --project-cl {options.project_cl}')
+
+            self.sync_engine_cl = options.engine_cl
+            self.sync_project_cl = options.project_cl
+            self.generate_after_sync = options.generate
+            self.clobber_engine = options.clobber_engine
+            self.clobber_project = options.clobber_project
+
+            try:
+                self.uproj_path = options.project.resolve(strict=True)
+            except FileNotFoundError:
+                self.parser.error(f'argument --project: FileNotFoundError: {options.project}')
+
+            try:
+                proj_data = json.load(self.uproj_path.open())
+            except json.JSONDecodeError:
+                self.parser.error(f'argument --project: unable to parse JSON\n\n{traceback.format_exc()}')
+
+            # Find engine; look up the tree from the project for an Engine/ dir
+            iter_dir = self.uproj_path.parent.parent
+            while not iter_dir.samefile(iter_dir.anchor):
+                candidate = iter_dir / "Engine"
+                if candidate.is_dir():
+                    self.engine_dir = candidate
+                    break
+                iter_dir = iter_dir.parent
+
+            if self.engine_dir is None:
+                if self.sync_engine_cl is not None:
+                    self.parser.error('--engine-cl was specified, but unable to determine engine directory')
+                elif self.generate_after_sync:
+                    self.parser.error('--generate was specified, but unable to determine engine directory')
+
+
+    def generate_project_files(self) -> int:
+        ''' Invokes GenerateProjectFiles script, streams the output to INFO logging, and returns the exit code. '''
+        assert self.uproj_path and self.engine_dir
+
+        if sys.platform.startswith('win'):
+            gpf_script = os.path.join(self.engine_dir, 'Build', 'BatchFiles', 'GenerateProjectFiles.bat')
+        else:
+            gpf_script = os.path.join(self.engine_dir, 'Build', 'BatchFiles', 'Linux' 'GenerateProjectFiles.sh')
+
+        args = [gpf_script, str(self.uproj_path)]
+
+        logging.debug(f'generate_project_files(): invoking subprocess: args={args}')
+
+        # stdin=subprocess.DEVNULL is required when launched via SwitchboardListener; otherwise this call tries
+        # to make the non-existent stdin inheritable, and raises `OSError: [WinError 6] The handle is invalid`
+        with subprocess.Popen(args, stdin=subprocess.DEVNULL, stdout=subprocess.PIPE) as proc:
+            for line in proc.stdout:
+                logging.info(f'gpf> {line.decode().rstrip()}')
+
+            return proc.wait()
+
+
+def main() -> int:
+    app = SbListenerHelper()
+    result = app.run()
+    logging.info(f'Finished with exit code {result}')
+    return result
+
+
+if __name__ == "__main__":
+    sys.exit(main())
diff --git a/Engine/Plugins/VirtualProduction/Switchboard/Source/Switchboard/switchboard/add_config_dialog.py b/Engine/Plugins/VirtualProduction/Switchboard/Source/Switchboard/switchboard/add_config_dialog.py
--- a/Engine/Plugins/VirtualProduction/Switchboard/Source/Switchboard/switchboard/add_config_dialog.py
+++ b/Engine/Plugins/VirtualProduction/Switchboard/Source/Switchboard/switchboard/add_config_dialog.py
@@ -5,6 +5,7 @@ from switchboard.config import CONFIG
 from .switchboard_logging import LOGGER
 
 import switchboard.p4_utils as P4
+import switchboard.switchboard_utils as sb_utils
 import os, sys, subprocess, shlex, pathlib, socket, threading
 
 class AddConfigDialog(QtWidgets.QDialog):
@@ -404,7 +405,7 @@ class AddConfigDialog(QtWidgets.QDialog):
                 # Relying on wmic for this. Another option is to use psutil which is cross-platform, and slower.
                 cmd = f'wmic process where caption="{UEname}" get commandline'
 
-                for line in subprocess.check_output(cmd, startupinfo=P4.get_sp_startupinfo()).decode().splitlines():
+                for line in subprocess.check_output(cmd, startupinfo=sb_utils.get_hidden_sp_startupinfo()).decode().splitlines():
                     if UEname.lower() not in line.lower():
                         continue
                     
diff --git a/Engine/Plugins/VirtualProduction/Switchboard/Source/Switchboard/switchboard/devices/device_base.py b/Engine/Plugins/VirtualProduction/Switchboard/Source/Switchboard/switchboard/devices/device_base.py
--- a/Engine/Plugins/VirtualProduction/Switchboard/Source/Switchboard/switchboard/devices/device_base.py
+++ b/Engine/Plugins/VirtualProduction/Switchboard/Source/Switchboard/switchboard/devices/device_base.py
@@ -38,7 +38,7 @@ class DeviceQtHandler(QtCore.QObject):
     signal_device_sync_failed = QtCore.Signal(object)
     signal_device_is_recording_device_changed = QtCore.Signal(object, int)
     signal_device_build_update = QtCore.Signal(object, str, str) # device, step, percent
-    signal_device_sync_update = QtCore.Signal(object, int) # device, progress
+    signal_device_sync_update = QtCore.Signal(object, str) # device, percent
 
 
 class Device(QtCore.QObject):
diff --git a/Engine/Plugins/VirtualProduction/Switchboard/Source/Switchboard/switchboard/devices/unreal/plugin_unreal.py b/Engine/Plugins/VirtualProduction/Switchboard/Source/Switchboard/switchboard/devices/unreal/plugin_unreal.py
--- a/Engine/Plugins/VirtualProduction/Switchboard/Source/Switchboard/switchboard/devices/unreal/plugin_unreal.py
+++ b/Engine/Plugins/VirtualProduction/Switchboard/Source/Switchboard/switchboard/devices/unreal/plugin_unreal.py
@@ -1,5 +1,23 @@
 # Copyright Epic Games, Inc. All Rights Reserved.
 
+from __future__ import annotations
+
+import base64
+import ipaddress
+import marshal
+import os
+import re
+import subprocess
+import sys
+import threading
+import uuid
+from abc import abstractmethod
+from collections import OrderedDict
+from functools import wraps
+from typing import Any, Callable, Dict, IO, List, Optional, Set
+
+from PySide2 import QtWidgets, QtGui, QtCore
+
 from switchboard import message_protocol, p4_utils, switchboard_application
 from switchboard import config_osc as osc
 from switchboard import switchboard_utils as sb_utils
@@ -13,19 +31,16 @@ from .listener_watcher import ListenerWatcher
 from .redeploy_dialog import RedeployListenerDialog
 from . import version_helpers
 
-from PySide2 import QtWidgets, QtGui, QtCore
-
-import base64, ipaddress, os, re, sys, threading, typing, uuid
-from collections import OrderedDict
-
 
-class ProgramStartQueueItem(object):
+class ProgramStartQueueItem:
     ''' Item that holds the information for a program to start in the future '''
-    def __init__(self, name, puuid_dependency, puuid, msg_to_unreal_client, pid=0, launch_fn=lambda:None):
+    def __init__(self, name: str, puuid_dependency: Optional[uuid.UUID], puuid: uuid.UUID,
+        msg_to_unreal_client: bytes, pid: int = 0, launch_fn: Callable[[], None] = lambda: None
+    ):
+        self.name = name
         self.puuid_dependency = puuid_dependency         # so that it can check wait_for_previous_to_end against the correct program
         self.puuid = puuid                               # id of the program.
         self.msg_to_unreal_client = msg_to_unreal_client # command for listener/unreal_client
-        self.name = name
         self.launch_fn = launch_fn
         self.pid = pid
 
@@ -34,34 +49,35 @@ class ProgramStartQueueItem(object):
         return cls(
             puuid_dependency = None,
             puuid = uuid.UUID(process['uuid']),
-            msg_to_unreal_client = '',
+            msg_to_unreal_client = b'',
             name = process['name'],
             pid = process['pid'],
-            launch_fn = None,
+            launch_fn = lambda: None,
         )
 
 
 def use_lock(func):
     ''' Decorator to ensure the decorated function is executed by a single thread at a time '''
+    @wraps(func)
     def _use_lock(self, *args, **kwargs):
         self.lock.acquire()
         try:
             return func(self, *args, **kwargs)
         finally:
             self.lock.release()
-            
+
     return _use_lock
 
 
-class ProgramStartQueue(object):
+class ProgramStartQueue:
     ''' Queue of programs to launch that may have dependencies '''
 
     def __init__(self):
         self.lock = threading.Lock() # Needed because these functions will be called from listener thread and main thread
 
-        self.queued_programs = [] # ProgramStartQueueItem
-        self.starting_programs = OrderedDict() # [puuid] = ProgramStartQueueItem 
-        self.running_programs = OrderedDict()  # [puuid] = ProgramStartQueueItem (initialized as returned by listener)
+        self.queued_programs: List[ProgramStartQueueItem] = []
+        self.starting_programs: OrderedDict[uuid.UUID, ProgramStartQueueItem] = OrderedDict()
+        self.running_programs: OrderedDict[uuid.UUID, ProgramStartQueueItem] = OrderedDict() # initialized as returned by listener
 
     @use_lock
     def reset(self):
@@ -139,7 +155,7 @@ class ProgramStartQueue(object):
 
             # check if it is queued
             if prog is None:
-                programs = [program for program in self.queued_programs if prog.puuid == puuid]
+                programs = [program for program in self.queued_programs if program.puuid == puuid]
 
                 for program in programs:
                     prog = program
@@ -291,13 +307,13 @@ class DeviceUnreal(Device):
     listener_watcher = ListenerWatcher()
 
     # Every DeviceUnreal (and derived class, e.g. DevicenDisplay) instance; used for listener updates.
-    active_unreal_devices: typing.Set['DeviceUnreal'] = set()
+    active_unreal_devices: Set[DeviceUnreal] = set()
 
     # Flag used to batch together multiple rapid calls to `_queue_notify_redeploy`.
     _pending_notify_redeploy = False
 
     @classmethod
-    def get_designated_local_builder(cls) -> typing.Optional['DeviceUnreal']:
+    def get_designated_local_builder(cls) -> Optional[DeviceUnreal]:
         ''' Returns first (by IP) loopback `DeviceUnreal` (or derived class) device.
             This is the device tasked with building the multiuser server and listener executables.
         '''
@@ -357,8 +373,6 @@ class DeviceUnreal(Device):
         self.inflight_project_cl = None
         self.inflight_engine_cl = None
 
-        self.sync_progress = 0
-
         # Set a delegate method if the device gets a disconnect signal
         self.unreal_client.disconnect_delegate = self.on_listener_disconnect
         self.unreal_client.receive_file_completed_delegate = self.on_file_received
@@ -369,7 +383,7 @@ class DeviceUnreal(Device):
         self.unreal_client.delegates["program ended"] = self.on_program_ended
         self.unreal_client.delegates["program started"] = self.on_program_started
         self.unreal_client.delegates["program killed"] = self.on_program_killed
-        self.unreal_client.delegates["start"] = self.on_program_started
+        self.unreal_client.delegates["start"] = self.on_program_started # This catches start failures.
 
         self.osc_connection_timer = QtCore.QTimer(self)
         self.osc_connection_timer.timeout.connect(self._try_connect_osc)
@@ -381,9 +395,6 @@ class DeviceUnreal(Device):
         # keeps track of programs to start
         self.program_start_queue = ProgramStartQueue()
 
-        # determines if it should force clobber in workspace
-        self.force_clobber = False
-
     @classmethod
     def added_device(cls, device):
         ''' Implementation of base class function. Called when one of our plugin devices has been added. '''
@@ -531,9 +542,9 @@ class DeviceUnreal(Device):
         working_dir = os.path.dirname(CONFIG.UPROJECT_PATH.get_value(self.name))
 
         puuid, msg = message_protocol.create_start_process_message(
-            prog_path="p4", 
-            prog_args=args, 
-            prog_name=program_name, 
+            prog_path="p4",
+            prog_args=args,
+            prog_name=program_name,
             caller=self.name,
             working_dir=working_dir,
             update_clients_with_stdout=False,
@@ -576,207 +587,73 @@ class DeviceUnreal(Device):
         super().disconnect_listener()
         self.unreal_client.disconnect()
 
-    def sync(self, engine_cl, project_cl):
+    def sync(self, engine_cl: Optional[int], project_cl: Optional[int]):
         if not engine_cl and not project_cl:
             LOGGER.warning(f"Neither project nor engine changelist is selected. There is nothing to sync!")
             return
 
-        if engine_cl:
-            self._sync_engine(engine_cl)
-            self.status = DeviceStatus.SYNCING
-
-        if project_cl:
-            self._sync_project(project_cl)
-            if self.status != DeviceStatus.SYNCING:
-                self.status = DeviceStatus.SYNCING
-
-    def _sync_engine(self, engine_cl):
-        program_name = "sync_engine"
+        program_name = 'sync'
 
         # check if it is already on its way:
         try:
             existing_puuid = self.program_start_queue.puuid_from_name(program_name)
-            LOGGER.info(f"{self.name}: Already syncing engine with puuid {existing_puuid}")
-            return
-        except KeyError:
-            pass
-
-        project_name = os.path.basename(os.path.dirname(CONFIG.UPROJECT_PATH.get_value(self.name)))
-        LOGGER.info(f"{self.name}: Queuing {program_name} for project {project_name} to revision {engine_cl}")
-
-        self.inflight_engine_cl = engine_cl
-
-        # when no uproject is specified to the UAT SyncProject script, it will go into EngineOnly mode and only sync the engine
-        # as well as all loose files one directoy up from the engine directory.
-        sync_tool = f'{os.path.normpath(os.path.join(CONFIG.ENGINE_DIR.get_value(self.name), "Build", "BatchFiles", "RunUAT.bat"))}'
-        sync_args = f'-P4 SyncProject -cl={engine_cl} -threads=8 -generate'
-
-        # check if project sync is already happening:
-        try:
-            sync_project_puuid = self.program_start_queue.puuid_from_name('sync_project')
-        except KeyError:
-            sync_project_puuid = None
-
-        puuid, msg = message_protocol.create_start_process_message(
-            prog_path=sync_tool, 
-            prog_args=sync_args, 
-            prog_name=program_name, 
-            caller=self.name,
-            update_clients_with_stdout=True,
-        )
-
-        self.program_start_queue.add(
-            ProgramStartQueueItem(
-                name = program_name,
-                puuid_dependency = sync_project_puuid,
-                puuid = puuid,
-                msg_to_unreal_client = msg,
-                launch_fn= lambda : LOGGER.info(f"{self.name}: Sending engine sync command: {sync_tool} {sync_args}")
-            ),
-            unreal_client = self.unreal_client,
-        )
-
-
-    def _sync_project(self, project_cl):
-        sync_project_prog_name = "sync_project"
-
-        # check if it is already on its way:
-        try:
-            existing_puuid = self.program_start_queue.puuid_from_name(sync_project_prog_name)
             LOGGER.info(f"{self.name}: Already syncing with puuid {existing_puuid}")
             return
         except KeyError:
             pass
 
         project_name = os.path.basename(os.path.dirname(CONFIG.UPROJECT_PATH.get_value(self.name)))
-        LOGGER.info(f"{self.name}: Queueing {sync_project_prog_name} {project_name} to revision {project_cl}")
-        self.inflight_project_cl = project_cl
-
-        sync_tool = ""
-        sync_args = ""
-
-        # todo: UAT requires help finding the correct workspace.
-        # the only way this can be done is by setting the environment variable uebp_CLIENT to whatever we have in CONFIG.SOURCE_CONTROL_WORKSPACE.
-        # however that needs to happen on the listener machines, so we would need a way to (platform-independently) set env variables when
-        # running commands and extend the protocol to include setting env vars as well.
-        # alternatively we could try running commands through cmd /V /C "set uebp_CLIENT=CONFIG.SOURCE_CONTROL_WORKSPACE&& EXE ARGS"
-        # but that would only work for Windows and thus require us to know what the listener machine's OS is.
-        engine_under_source_control = CONFIG.BUILD_ENGINE.get_value()
+        LOGGER.info(f"{self.name}: Queuing sync for project {project_name} (revisions: engine={engine_cl}, project={project_cl})")
 
         workspace = CONFIG.SOURCE_CONTROL_WORKSPACE.get_value(self.name)
+        engine_under_source_control = CONFIG.BUILD_ENGINE.get_value()
 
-        if engine_under_source_control:
-            sync_tool = f'{os.path.normpath(os.path.join(CONFIG.ENGINE_DIR.get_value(self.name), "Build", "BatchFiles", "RunUAT.bat"))}'
-            sync_args = f'-P4 SyncProject -project="{CONFIG.UPROJECT_PATH.get_value(self.name)}" -projectonly -cl={project_cl} -threads=8 -generate'
-        else:
-            # for installed/vanilla engine we directly call p4 to sync the project itself.
-            p4_path = CONFIG.P4_PROJECT_PATH.get_value()
-            sync_tool = "p4"
-            sync_args = f"-c{workspace} sync {p4_path}/...@{project_cl}"
-
-        # check if engine sync is already happening:
-        try:
-            sync_engine_puuid = self.program_start_queue.puuid_from_name('sync_engine')
-        except KeyError:
-            sync_engine_puuid = None
+        if engine_under_source_control or engine_cl:
+            python_path = os.path.normpath(os.path.join(CONFIG.ENGINE_DIR.get_value(self.name),
+                'Binaries', 'ThirdParty', 'Python3', 'Win64', 'python.exe'))
+            helper_path = os.path.normpath(os.path.join(CONFIG.ENGINE_DIR.get_value(self.name),
+                'Plugins', 'VirtualProduction', 'Switchboard', 'Source', 'Switchboard', 'sbl_helper.py'))
 
-        # find out if it is noclobber
-        program_name = "isclobber"
+            sync_tool = python_path
+            sync_args = f'{helper_path} sync --project={CONFIG.UPROJECT_PATH.get_value(self.name)} --generate'
 
-        working_dir = os.path.dirname(CONFIG.UPROJECT_PATH.get_value(self.name))
-
-        puuid_isclobber, msg = message_protocol.create_start_process_message(
-            prog_path="p4", 
-            prog_args=f"client -o {workspace}",
-            prog_name=program_name, 
-            caller=self.name,
-            working_dir=working_dir,
-            update_clients_with_stdout=False,
-        )
+            if workspace:
+                sync_args += f' --p4client={workspace}'
 
-        self.program_start_queue.add(
-            ProgramStartQueueItem(
-                name = program_name,
-                puuid_dependency = sync_engine_puuid,
-                puuid = puuid_isclobber,
-                msg_to_unreal_client = msg,
-                launch_fn= lambda : LOGGER.info(f"{self.name}: Sending project sync command: {sync_tool} {sync_args}"),
-            ),
-            unreal_client = self.unreal_client,
-        )
+            if engine_cl:
+                sync_args += f' --engine-cl={engine_cl}'
+                self.inflight_engine_cl = engine_cl
 
-        # allow clobbering
-        if self.force_clobber and sys.platform.startswith('win'):
-            program_name = "clobber"
-
-            puuid_clobber, msg = message_protocol.create_start_process_message(
-                prog_path="powershell", 
-                prog_args=f"p4 client -o {workspace} | %{{$_ -replace 'noclobber', 'clobber'}} | p4 client -i",
-                prog_name=program_name, 
-                caller=self.name,
-                working_dir=working_dir,
-                update_clients_with_stdout=True,
-            )
-
-            self.program_start_queue.add(
-                ProgramStartQueueItem(
-                    name = program_name,
-                    puuid_dependency = puuid_isclobber,
-                    puuid = puuid_clobber,
-                    msg_to_unreal_client = msg,
-                ),
-                unreal_client = self.unreal_client,
-            )
+            if project_cl:
+                sync_args += f' --project-cl={project_cl} --clobber-project'
+                self.inflight_project_cl = project_cl
         else:
-            # TODO: Support other operating systems
-            puuid_clobber = puuid_isclobber # transfer puuid dependency to clobber puuid
+            # for installed/vanilla engine we directly call p4 to sync the project itself.
+            sync_tool = "p4"
+            sync_args = f"-c{workspace} sync {CONFIG.P4_PROJECT_PATH.get_value()}/...@{project_cl}"
 
-        # start sync
-        puuid_sync, msg = message_protocol.create_start_process_message(
+        puuid, msg = message_protocol.create_start_process_message(
             prog_path=sync_tool, 
             prog_args=sync_args, 
-            prog_name=sync_project_prog_name, 
+            prog_name=program_name, 
             caller=self.name,
-            working_dir=working_dir,
+            working_dir=os.path.dirname(CONFIG.UPROJECT_PATH.get_value(self.name)),
             update_clients_with_stdout=True,
         )
 
         self.program_start_queue.add(
             ProgramStartQueueItem(
-                name = sync_project_prog_name,
-                puuid_dependency = puuid_clobber,
-                puuid = puuid_sync,
+                name = program_name,
+                puuid_dependency = None,
+                puuid = puuid,
                 msg_to_unreal_client = msg,
+                launch_fn = lambda: LOGGER.info(f"{self.name}: Sending sync command: {sync_tool} {sync_args}")
             ),
             unreal_client = self.unreal_client,
         )
 
-        # disallow clobbering
-        if self.force_clobber and sys.platform.startswith('win'):
-
-            program_name = "noclobber"
-
-            puuid, msg = message_protocol.create_start_process_message(
-                prog_path="powershell", 
-                prog_args=f"p4 client -o {workspace} | %{{$_ -replace ' clobber', ' noclobber'}} | p4 client -i",
-                prog_name=program_name, 
-                caller=self.name,
-                working_dir=working_dir,
-                update_clients_with_stdout=True,
-            )
-
-            self.program_start_queue.add(
-                ProgramStartQueueItem(
-                    name = program_name,
-                    puuid_dependency = puuid_sync,
-                    puuid = puuid,
-                    msg_to_unreal_client = msg,
-                ),
-                unreal_client = self.unreal_client,
-            )
-        else:
-            # TODO: Support other operating systems
-            pass
+        if self.status != DeviceStatus.SYNCING:
+            self.status = DeviceStatus.SYNCING
 
     def build(self):
         program_name = 'build_project'
@@ -791,16 +668,14 @@ class DeviceUnreal(Device):
 
         # check for any sync dependencies
         sync_puuid = None
-        sync_name = ''
-        for sync_program_name in ['sync_engine', 'sync_project']: # order matters, as we sync the engine before the project.
-            try:
-                sync_puuid = self.program_start_queue.puuid_from_name(sync_program_name)
-                sync_name = sync_program_name
-            except KeyError:
-                pass
+
+        try:
+            sync_puuid = self.program_start_queue.puuid_from_name('sync')
+        except KeyError:
+            pass
 
         if sync_puuid:
-            LOGGER.debug(f"{self.name} Queuing build after {sync_name} with puuid {sync_puuid}")
+            LOGGER.debug(f"{self.name} Queuing build after sync with puuid {sync_puuid}")
 
         # Build dependency chain
         puuid_dependency = sync_puuid
@@ -829,19 +704,19 @@ class DeviceUnreal(Device):
 
         puuid_dependency = self._build_project(puuid_dependency=puuid_dependency)
 
-    def _build_project(self, puuid_dependency: typing.Optional[uuid.UUID] = None):
+    def _build_project(self, puuid_dependency: Optional[uuid.UUID] = None):
         ubt_args = f'Win64 Development -project="{CONFIG.UPROJECT_PATH.get_value(self.name)}" -TargetType=Editor -Progress -NoHotReloadFromIDE'
         return self._queue_build('project', ubt_args=ubt_args, puuid_dependency=puuid_dependency)
 
-    def _build_mu_server(self, puuid_dependency: typing.Optional[uuid.UUID] = None):
+    def _build_mu_server(self, puuid_dependency: Optional[uuid.UUID] = None):
         ubt_args = 'UnrealMultiUserServer Win64 Development -Progress'
         return self._queue_build('mu_server', ubt_args=ubt_args, puuid_dependency=puuid_dependency)
 
-    def _build_listener(self, puuid_dependency: typing.Optional[uuid.UUID] = None):
+    def _build_listener(self, puuid_dependency: Optional[uuid.UUID] = None):
         ubt_args = 'SwitchboardListener Win64 Development -Progress'
         return self._queue_build('listener', ubt_args=ubt_args, puuid_dependency=puuid_dependency)
 
-    def _queue_build(self, program_name_suffix: str, ubt_args: str, puuid_dependency: typing.Optional[uuid.UUID] = None):
+    def _queue_build(self, program_name_suffix: str, ubt_args: str, puuid_dependency: Optional[uuid.UUID] = None):
         program_name = f"build_{program_name_suffix}"
 
         engine_path = CONFIG.ENGINE_DIR.get_value(self.name)
@@ -986,10 +861,10 @@ class DeviceUnreal(Device):
         args = args.replace('\r', ' ').replace('\n', ' ')
 
         puuid, msg = message_protocol.create_start_process_message(
-            prog_path=engine_path, 
-            prog_args=args, 
-            prog_name=program_name, 
-            caller=self.name, 
+            prog_path=engine_path,
+            prog_args=args,
+            prog_name=program_name,
+            caller=self.name,
             update_clients_with_stdout=False,
             priority_modifier=priority_modifier
         )
@@ -1022,8 +897,7 @@ class DeviceUnreal(Device):
             self.unreal_started_signal.emit()
         elif prog.name.startswith('build_'):
             self.status = DeviceStatus.BUILDING
-        elif prog.name.startswith('sync_'):
-            self.sync_progress = 0
+        elif prog.name == 'sync':
             self.status = DeviceStatus.SYNCING
 
         self.program_start_queue.update_running_program(prog=prog)
@@ -1041,7 +915,7 @@ class DeviceUnreal(Device):
             # log this
             LOGGER.error(f"Could not start {program_name}: {message['error']}")
 
-            if program_name.startswith('sync_') or program_name.startswith('build_'):
+            if program_name == 'sync' or program_name.startswith('build_'):
                 self.status = DeviceStatus.CLOSED
                 self.project_changelist = self.project_changelist # force to show existing project_changelist to hide building/syncing
             elif program_name == 'unreal':
@@ -1086,7 +960,10 @@ class DeviceUnreal(Device):
 
         puuid = uuid.UUID(process['uuid'])
         returncode = message['returncode']
-        output = message['output']
+
+        def get_stdout_str() -> str:
+            b64bytes = base64.b64decode(message['stdoutB64'])
+            return b64bytes.decode()
 
         LOGGER.info(f"{self.name}: Program with id {puuid} exited with returncode {returncode}")
 
@@ -1104,54 +981,34 @@ class DeviceUnreal(Device):
         if program_name == "unreal" and not len(remaining_homonyms):
             self.status = DeviceStatus.CLOSED
 
-        elif program_name.startswith('sync_'):
-            if "sync_engine" == program_name:
-                if returncode == 0:
-                    LOGGER.info(f"{self.name}: Engine was synced successfully")
-                else:
-                    LOGGER.error(f"{self.name}: Engine was not synced successfully!")
-                    for line in output.splitlines():
-                        LOGGER.error(f"{self.name}: {line}")
-
-                    # flag the inflight cl as invalid
-                    self.inflight_engine_cl = None
-
-                    # notify of the failure
-                    self.device_qt_handler.signal_device_sync_failed.emit(self)
+        elif program_name == 'sync':
+            if returncode == 0:
+                LOGGER.info(f"{self.name}: Sync successful")
+            else:
+                LOGGER.error(f"{self.name}: Sync failed!")
+                for line in get_stdout_str().splitlines():
+                    LOGGER.error(f"{self.name}: {line}")
 
-                # Update CL with the one in flight
-                self.engine_changelist = self.inflight_engine_cl
+                # flag the inflight cl as invalid
                 self.inflight_engine_cl = None
+                self.inflight_project_cl = None
 
-                # refresh project CL
-                self.project_changelist = self.project_changelist
-
-            elif "sync_project" == program_name:
-                if returncode == 0:
-                    LOGGER.info(f"{self.name}: Project was synced successfully")
-                else:
-                    LOGGER.error(f"{self.name}: Project was not synced successfully!")
-                    for line in output.splitlines():
-                        LOGGER.error(f"{self.name}: {line}")
-
-                    # flag the inflight cl as invalid
-                    self.inflight_project_cl = None
+                # notify of the failure
+                self.device_qt_handler.signal_device_sync_failed.emit(self)
 
-                    # notify of the failure
-                    self.device_qt_handler.signal_device_sync_failed.emit(self)
+            # If you build and sync the engine, update its CL
+            if CONFIG.BUILD_ENGINE.get_value():
+                self.engine_changelist = self.inflight_engine_cl if self.inflight_engine_cl != None else self.engine_changelist
+                self.inflight_engine_cl = None
 
-                # Update CL with the one in flight
-                self.project_changelist = self.inflight_project_cl
-                self.inflight_project_cl = None
+            # Update CL with the one in flight
+            self.project_changelist = self.inflight_project_cl
+            self.inflight_project_cl = None
 
-                # If you build and sync the engine, update its CL
-                if CONFIG.BUILD_ENGINE.get_value():
-                    self.engine_changelist = self.inflight_engine_cl if self.inflight_engine_cl != None else self.engine_changelist
-                    self.inflight_engine_cl = None
+            # refresh project CL
+            self.project_changelist = self.project_changelist
 
-            if self.inflight_project_cl == None and self.inflight_engine_cl == None:
-                # everything is done syncing
-                self.status = DeviceStatus.CLOSED
+            self.status = DeviceStatus.CLOSED
 
         elif program_name.startswith('build_'):
             if returncode == 0:
@@ -1161,7 +1018,7 @@ class DeviceUnreal(Device):
 
                 # MSVC build tools error codes, e.g. 'error C4430' or 'error LNK1104'
                 error_pattern = re.compile(r"error [A-Z]{1,3}[0-9]{4}")
-                for line in output.splitlines():
+                for line in get_stdout_str().splitlines():
                     if error_pattern.search(line):
                         LOGGER.error(f"{self.name}: {line}")
 
@@ -1172,6 +1029,7 @@ class DeviceUnreal(Device):
                     self.engine_changelist = self.engine_changelist
 
         elif "cstat" in program_name:
+            output = get_stdout_str()
             changelists = [line.strip() for line in output.split()]
 
             try:
@@ -1195,18 +1053,12 @@ class DeviceUnreal(Device):
                 LOGGER.info(f"{self.name}: Engine used for project {project_name} is on revision {current_changelist}")
                 self.engine_changelist = current_changelist
 
-        elif program_name == 'isclobber':
-            for line in output.splitlines():
-                if line.startswith('Options:'):
-                    self.force_clobber = True if 'noclobber' in line else False
-                    break
 
     def on_program_killed(self, message):
         ''' Handler of killed program. Expect on_program_ended for anything other than a fail. '''
         if not message['bAck']:
 
             # remove from list of puuids (if it exists)
-            #
             puuid = uuid.UUID(message['puuid'])
             self.program_start_queue.on_program_ended(puuid=puuid, unreal_client=self.unreal_client)
 
@@ -1238,8 +1090,9 @@ class DeviceUnreal(Device):
         if process['caller'] != self.name:
             return
 
-        stdout = ''.join([chr(b) for b in message['partialStdout'] if b != ord('\r')])
-        lines = list(filter(None, stdout.split('\n')))
+        stdoutbytes = base64.b64decode(message['partialStdoutB64'])
+        stdout = stdoutbytes.decode()
+        lines = list(filter(None, stdout.splitlines()))
 
         # see if this is an update to the build
         #
@@ -1262,13 +1115,17 @@ class DeviceUnreal(Device):
                     step = stepparts[-2].strip()
 
                     percent = line.split(' ')[-1].strip()
-                    
+
                     if '%' == percent[-1]:
                         self.device_qt_handler.signal_device_build_update.emit(self, step, percent)
 
-        elif process['name'].startswith('sync_'):
-            self.sync_progress += len(lines)
-            self.device_qt_handler.signal_device_sync_update.emit(self, self.sync_progress)
+        elif process['name'] == 'sync':
+            for line in lines:
+                if 'Progress:' in line:
+                    match = re.search(r'Progress: (\d{1,3}\.\d\d%)', line)
+                    if match:
+                        sync_progress = match.group(1)
+                        self.device_qt_handler.signal_device_sync_update.emit(self, sync_progress)
 
         for line in lines:
             LOGGER.debug(f"{self.name} {process['name']}: {line}")
@@ -1507,8 +1364,8 @@ class DeviceWidgetUnreal(DeviceWidget):
 
         self.project_changelist_label.show()
 
-    def update_sync_status(self, device, linecount):
-        self.project_changelist_label.setText(f"Syncing...{linecount}")
+    def update_sync_status(self, device, percent):
+        self.project_changelist_label.setText(f"Syncing...{percent}")
         self.project_changelist_label.setToolTip('Syncing from Version Control')
 
         self.project_changelist_label.show()
diff --git a/Engine/Plugins/VirtualProduction/Switchboard/Source/Switchboard/switchboard/devices/unreal/version_helpers.py b/Engine/Plugins/VirtualProduction/Switchboard/Source/Switchboard/switchboard/devices/unreal/version_helpers.py
--- a/Engine/Plugins/VirtualProduction/Switchboard/Source/Switchboard/switchboard/devices/unreal/version_helpers.py
+++ b/Engine/Plugins/VirtualProduction/Switchboard/Source/Switchboard/switchboard/devices/unreal/version_helpers.py
@@ -3,9 +3,12 @@
 import typing
 
 
-LISTENER_COMPATIBLE_VERSION = (1,4)
+LISTENER_COMPATIBLE_VERSION = (1,5)
+
+# Version where redeploy support was initially introduced.
 LISTENER_MIN_REDEPLOY_VERSION = (1,4,0)
 
+
 ListenerVersion = typing.Tuple[int, int, int]
 
 
diff --git a/Engine/Plugins/VirtualProduction/Switchboard/Source/Switchboard/switchboard/message_protocol.py b/Engine/Plugins/VirtualProduction/Switchboard/Source/Switchboard/switchboard/message_protocol.py
--- a/Engine/Plugins/VirtualProduction/Switchboard/Source/Switchboard/switchboard/message_protocol.py
+++ b/Engine/Plugins/VirtualProduction/Switchboard/Source/Switchboard/switchboard/message_protocol.py
@@ -1,28 +1,33 @@
 # Copyright Epic Games, Inc. All Rights Reserved.
 
-import base64, json, uuid
+import base64
+import json
+import uuid
+
 
 def create_start_process_message(
     prog_path: str,
     prog_args: str,
     prog_name: str,
     caller: str,
-    update_clients_with_stdout: bool, 
     working_dir: str = "",
-    priority_modifier: int = 0
+    *,
+	update_clients_with_stdout: bool = False,
+    priority_modifier: int = 0,
 ):
     cmd_id = uuid.uuid4()
     start_cmd = {
-        'command': 'start', 
-        'id': str(cmd_id), 
-        'exe': prog_path, 
-        'args': prog_args, 
-        'name': prog_name, 
+        'command': 'start',
+        'id': str(cmd_id),
+        'exe': prog_path,
+        'args': prog_args,
+        'name': prog_name,
         'caller': caller,
         'working_dir': working_dir,
-        'bUpdateClientsWithStdout' : update_clients_with_stdout,
+        'bUpdateClientsWithStdout': update_clients_with_stdout,
         'priority_modifier': priority_modifier,
     }
+
     message = json.dumps(start_cmd).encode() + b'\x00'
     return (cmd_id, message)
 
diff --git a/Engine/Plugins/VirtualProduction/Switchboard/Source/Switchboard/switchboard/p4_utils.py b/Engine/Plugins/VirtualProduction/Switchboard/Source/Switchboard/switchboard/p4_utils.py
--- a/Engine/Plugins/VirtualProduction/Switchboard/Source/Switchboard/switchboard/p4_utils.py
+++ b/Engine/Plugins/VirtualProduction/Switchboard/Source/Switchboard/switchboard/p4_utils.py
@@ -1,20 +1,13 @@
 # Copyright Epic Games, Inc. All Rights Reserved.
 
-import subprocess, socket, re, os, marshal, sys, pathlib
-
-from .switchboard_logging import LOGGER
+import marshal
+import pathlib
+import subprocess
 from functools import wraps
-from .config import CONFIG
-
-def get_sp_startupinfo():
-    ''' Returns subprocess.startupinfo and avoids extra cmd line window in windows.
-    '''
-    startupinfo = subprocess.STARTUPINFO()
 
-    if sys.platform.startswith("win"):
-        startupinfo.dwFlags |= subprocess.STARTF_USESHOWWINDOW
+from . import switchboard_utils as sb_utils
+from .switchboard_logging import LOGGER
 
-    return startupinfo
 
 def p4_login(f):
     @wraps(f)
@@ -29,30 +22,6 @@ def p4_login(f):
     return wrapped
 
 
-@p4_login
-def p4_stream_root(client):
-    """ Returns stream root of client. """
-    p4_command = f'p4 -ztag -F "%Stream%" -c {client} stream -o'
-    LOGGER.info(f"Executing: {p4_command}")
-
-    p4_result = subprocess.check_output(p4_command, startupinfo=get_sp_startupinfo()).decode()
-    if p4_result:
-        return p4_result.strip()
-    return None
-
-
-@p4_login
-def p4_where(client, local_path):
-    """Returns depot path of local file."""
-    p4_command = f'p4 -ztag -c {client} -F "%depotFile%" where {local_path}'
-    LOGGER.info(f"Executing: {p4_command}")
-
-    p4_result = subprocess.check_output(p4_command, startupinfo=get_sp_startupinfo()).decode()
-    if p4_result:
-        return p4_result.strip()
-    return None
-
-
 @p4_login
 def p4_latest_changelist(p4_path, working_dir, num_changelists=10):
     """
@@ -61,34 +30,13 @@ def p4_latest_changelist(p4_path, working_dir, num_changelists=10):
     p4_command = f'p4 -ztag -F "%change%" changes -m {num_changelists} {p4_path}/...'
     LOGGER.info(f"Executing: {p4_command}")
 
-    p4_result = subprocess.check_output(p4_command, cwd=working_dir, startupinfo=get_sp_startupinfo()).decode()
+    p4_result = subprocess.check_output(p4_command, cwd=working_dir, startupinfo=sb_utils.get_hidden_sp_startupinfo()).decode()
 
     if p4_result:
         return p4_result.split()
 
     return None
 
-
-@p4_login
-def p4_current_user_name():
-    p4_command = f'p4 set P4USER'
-
-    p4_result = subprocess.check_output(p4_command, startupinfo=get_sp_startupinfo()).decode().rstrip()
-
-    p = re.compile("P4USER=(.*)\\(set\\)")
-    matches = p.search(p4_result)
-    if matches:
-        return matches.group(1).rstrip()
-    return None
-
-
-@p4_login
-def p4_edit(file_path):
-    p4_command = f'p4 edit "{file_path}"'
-
-    p4_result = subprocess.check_output(p4_command, startupinfo=get_sp_startupinfo()).decode()
-    LOGGER.debug(p4_result)
-
 def run(cmd, args=[], input=None):
     ''' Runs the provided p4 command and arguments with -G python marshaling 
     Args:
diff --git a/Engine/Plugins/VirtualProduction/Switchboard/Source/Switchboard/switchboard/switchboard_logging.py b/Engine/Plugins/VirtualProduction/Switchboard/Source/Switchboard/switchboard/switchboard_logging.py
--- a/Engine/Plugins/VirtualProduction/Switchboard/Source/Switchboard/switchboard/switchboard_logging.py
+++ b/Engine/Plugins/VirtualProduction/Switchboard/Source/Switchboard/switchboard/switchboard_logging.py
@@ -98,7 +98,6 @@ class HTMLogger(object):
         self.logger = logger
         self.file_handler = None
 
-        self.info(self.logger.getEffectiveLevel())
         self.add_custom_levels()
 
     def add_custom_levels(self):
diff --git a/Engine/Source/Programs/SwitchboardListener/Private/SwitchboardListener.cpp b/Engine/Source/Programs/SwitchboardListener/Private/SwitchboardListener.cpp
--- a/Engine/Source/Programs/SwitchboardListener/Private/SwitchboardListener.cpp
+++ b/Engine/Source/Programs/SwitchboardListener/Private/SwitchboardListener.cpp
@@ -1886,6 +1886,35 @@ void FSwitchboardListener::DisconnectClient(const FIPv4Endpoint& InClientEndpoin
 	ReceiveBuffer.Remove(InClientEndpoint);
 }
 
+void FSwitchboardListener::HandleStdout(const TSharedPtr<FRunningProcess, ESPMode::ThreadSafe>& Process)
+{
+	TArray<uint8> Output;
+	if (FPlatformProcess::ReadPipeToArray(Process->ReadPipe, Output))
+	{
+		Process->Output.Append(Output);
+	}
+
+	// If there was a new stdout, update the clients
+	if (Output.Num() && Process->bUpdateClientsWithStdout)
+	{
+		FSwitchboardProgramStdout Packet;
+
+		Packet.Process.Uuid = Process->UUID.ToString();
+		Packet.Process.Name = Process->Name;
+		Packet.Process.Path = Process->Path;
+		Packet.Process.Caller = Process->Caller;
+		Packet.Process.Pid = Process->PID;
+
+		Packet.PartialStdoutB64 = FBase64::Encode(Output);
+
+		for (const TPair<FIPv4Endpoint, TSharedPtr<FSocket>>& Connection : Connections)
+		{
+			const FIPv4Endpoint& ClientEndpoint = Connection.Key;
+			SendMessage(CreateMessage(Packet), ClientEndpoint);
+		}
+	}
+}
+
 void FSwitchboardListener::HandleRunningProcesses(TArray<TSharedPtr<FRunningProcess, ESPMode::ThreadSafe>>& Processes, bool bNotifyThatProgramEnded)
 {
 	// Reads pipe and cleans up dead processes from the array.
@@ -1902,50 +1931,20 @@ void FSwitchboardListener::HandleRunningProcesses(TArray<TSharedPtr<FRunningProc
 
 		if (Process->Handle.IsValid())
 		{
-			TArray<uint8> Output;
-
-			if (FPlatformProcess::ReadPipeToArray(Process->ReadPipe, Output))
-			{
-				// make sure the output array always has exactly one trailing null terminator.
-				// this way we can always convert to a valid string.
-				if (Process->Output.Num() > 0)
-				{
-					Process->Output.RemoveAt(Process->Output.Num() - 1);
-				}
-				Process->Output.Append(Output);
-				Process->Output.Add('\x00');
-			}
-
-			// If there was a new stdout, update the clients
-			if (Output.Num() && Process->bUpdateClientsWithStdout)
-			{
-				FSwitchboardProgramStdout Packet;
-
-				Packet.Process.Uuid = Process->UUID.ToString();
-				Packet.Process.Name = Process->Name;
-				Packet.Process.Path = Process->Path;
-				Packet.Process.Caller = Process->Caller;
-				Packet.Process.Pid = Process->PID;
-
-				Packet.PartialStdout = Output;
-
-				for (const TPair<FIPv4Endpoint, TSharedPtr<FSocket>>& Connection : Connections)
-				{
-					const FIPv4Endpoint& ClientEndpoint = Connection.Key;
-					SendMessage(CreateMessage(Packet), ClientEndpoint);
-				}
-			}
+			HandleStdout(Process);
 
 			if (!FPlatformProcess::IsProcRunning(Process->Handle))
 			{
+				// A final read can be necessary to avoid truncating the output.
+				HandleStdout(Process);
+
 				int32 ReturnCode = 0;
 				FPlatformProcess::GetProcReturnCode(Process->Handle, &ReturnCode);
 				UE_LOG(LogSwitchboard, Display, TEXT("Process exited with returncode: %d"), ReturnCode);
 
-				const FString ProcessOutput(UTF8_TO_TCHAR(Process->Output.GetData()));
 				if (ReturnCode != 0)
 				{
-					UE_LOG(LogSwitchboard, Display, TEXT("Output:\n%s"), *ProcessOutput);
+					UE_LOG(LogSwitchboard, Display, TEXT("Output:\n%s"), UTF8_TO_TCHAR(Process->Output.GetData()));
 				}
 
 				// Notify remote client, which implies that this is a program managed by it.
@@ -1959,7 +1958,7 @@ void FSwitchboardListener::HandleRunningProcesses(TArray<TSharedPtr<FRunningProc
 					Packet.Process.Caller = Process->Caller;
 					Packet.Process.Pid = Process->PID;
 					Packet.Returncode = ReturnCode;
-					Packet.Output = ProcessOutput;
+					Packet.StdoutB64 = FBase64::Encode(Process->Output);
 
 					for (const TPair<FIPv4Endpoint, TSharedPtr<FSocket>>& Connection : Connections)
 					{
diff --git a/Engine/Source/Programs/SwitchboardListener/Private/SwitchboardListener.h b/Engine/Source/Programs/SwitchboardListener/Private/SwitchboardListener.h
--- a/Engine/Source/Programs/SwitchboardListener/Private/SwitchboardListener.h
+++ b/Engine/Source/Programs/SwitchboardListener/Private/SwitchboardListener.h
@@ -127,6 +127,7 @@ private:
 
 	void CleanUpDisconnectedSockets();
 	void DisconnectClient(const FIPv4Endpoint& InClientEndpoint);
+	void HandleStdout(const TSharedPtr<FRunningProcess, ESPMode::ThreadSafe>& Process);
 	void HandleRunningProcesses(TArray<TSharedPtr<FRunningProcess, ESPMode::ThreadSafe>>& Processes, bool bNotifyThatProgramEnded);
 
 	bool SendMessage(const FString& InMessage, const FIPv4Endpoint& InEndpoint);
diff --git a/Engine/Source/Programs/SwitchboardListener/Private/SwitchboardListenerVersion.h b/Engine/Source/Programs/SwitchboardListener/Private/SwitchboardListenerVersion.h
--- a/Engine/Source/Programs/SwitchboardListener/Private/SwitchboardListenerVersion.h
+++ b/Engine/Source/Programs/SwitchboardListener/Private/SwitchboardListenerVersion.h
@@ -3,5 +3,5 @@
 #pragma once
 
 #define SBLISTENER_VERSION_MAJOR 1
-#define SBLISTENER_VERSION_MINOR 4
-#define SBLISTENER_VERSION_PATCH 5
+#define SBLISTENER_VERSION_MINOR 5
+#define SBLISTENER_VERSION_PATCH 0
diff --git a/Engine/Source/Programs/SwitchboardListener/Private/SwitchboardPacket.h b/Engine/Source/Programs/SwitchboardListener/Private/SwitchboardPacket.h
--- a/Engine/Source/Programs/SwitchboardListener/Private/SwitchboardPacket.h
+++ b/Engine/Source/Programs/SwitchboardListener/Private/SwitchboardPacket.h
@@ -98,7 +98,7 @@ struct FSwitchboardProgramStdout : public FSwitchboardPacket
 	FSwitchboardStateRunningProcess Process;
 
 	UPROPERTY()
-	TArray<uint8> PartialStdout;
+	FString PartialStdoutB64;
 };
 
 USTRUCT()
@@ -119,7 +119,7 @@ struct FSwitchboardProgramEnded : public FSwitchboardPacket
 	int32 Returncode;
 
 	UPROPERTY()
-	FString Output;
+	FString StdoutB64;
 };
 
 USTRUCT()
diff --git a/Engine/Source/Programs/SwitchboardListener/Private/SwitchboardProtocol.cpp b/Engine/Source/Programs/SwitchboardListener/Private/SwitchboardProtocol.cpp
--- a/Engine/Source/Programs/SwitchboardListener/Private/SwitchboardProtocol.cpp
+++ b/Engine/Source/Programs/SwitchboardListener/Private/SwitchboardProtocol.cpp
@@ -155,31 +155,33 @@ bool CreateTaskFromCommand(const FString& InCommand, const FIPv4Endpoint& InEndp
 		TSharedPtr<FJsonValue> NameField = TryGetCommandRequiredField(JsonData, TEXT("name"));
 		TSharedPtr<FJsonValue> CallerField = TryGetCommandRequiredField(JsonData, TEXT("caller"));
 		TSharedPtr<FJsonValue> WorkingDirField = TryGetCommandRequiredField(JsonData, TEXT("working_dir"));
-		TSharedPtr<FJsonValue> UpdateClientsWithStdoutField = TryGetCommandRequiredField(JsonData, TEXT("bUpdateClientsWithStdout"));
 
-		if (!ExeField || !ArgsField || !NameField || !CallerField || !WorkingDirField || !UpdateClientsWithStdoutField)
+		if (!ExeField || !ArgsField || !NameField || !CallerField || !WorkingDirField)
 		{
 			return false;
 		}
 
-		int32 PriorityModifier = 0;
-		if (TSharedPtr<FJsonValue> PriorityModifierField = JsonData->TryGetField(TEXT("priority_modifier")))
-		{
-			PriorityModifierField->TryGetNumber(PriorityModifier);
-		}
-
-		OutTask = MakeUnique<FSwitchboardStartTask>(
+		TUniquePtr<FSwitchboardStartTask> Task = MakeUnique<FSwitchboardStartTask>(
 			MessageID,
 			InEndpoint,
 			ExeField->AsString(),
 			ArgsField->AsString(),
 			NameField->AsString(),
 			CallerField->AsString(),
-			WorkingDirField->AsString(),
-			UpdateClientsWithStdoutField->AsBool(),
-			PriorityModifier
+			WorkingDirField->AsString()
 		);
 
+		if (TSharedPtr<FJsonValue> UpdateClientsWithStdoutField = JsonData->TryGetField(TEXT("bUpdateClientsWithStdout")))
+		{
+			UpdateClientsWithStdoutField->TryGetBool(Task->bUpdateClientsWithStdout);
+		}
+
+		if (TSharedPtr<FJsonValue> PriorityModifierField = JsonData->TryGetField(TEXT("priority_modifier")))
+		{
+			PriorityModifierField->TryGetNumber(Task->PriorityModifier);
+		}
+
+		OutTask = MoveTemp(Task);
 		return true;
 	}
 	else if (CommandName == TEXT("kill"))
diff --git a/Engine/Source/Programs/SwitchboardListener/Private/SwitchboardTasks.h b/Engine/Source/Programs/SwitchboardListener/Private/SwitchboardTasks.h
--- a/Engine/Source/Programs/SwitchboardListener/Private/SwitchboardTasks.h
+++ b/Engine/Source/Programs/SwitchboardListener/Private/SwitchboardTasks.h
@@ -88,9 +88,7 @@ struct FSwitchboardStartTask : public FSwitchboardTask
 		const FString& InArgs, 
 		const FString& InName, 
 		const FString& InCaller, 
-		const FString& InWorkingDir,
-		bool bInUpdateClientsWithStdout,
-		int32 InPriorityModifier
+		const FString& InWorkingDir
 	)
 		: FSwitchboardTask{ ESwitchboardTaskType::Start, TEXT("start"), InTaskId, InEndpoint }
 		, Command(InCommand)
@@ -98,8 +96,6 @@ struct FSwitchboardStartTask : public FSwitchboardTask
 		, Name(InName)
 		, Caller(InCaller)
 		, WorkingDir(InWorkingDir)
-		, bUpdateClientsWithStdout(bInUpdateClientsWithStdout)
-		, PriorityModifier(InPriorityModifier)
 	{
 	}
 
@@ -108,8 +104,8 @@ struct FSwitchboardStartTask : public FSwitchboardTask
 	FString Name;
 	FString Caller;
 	FString WorkingDir;
-	bool bUpdateClientsWithStdout;
-	int32 PriorityModifier;
+	bool bUpdateClientsWithStdout = false;
+	int32 PriorityModifier = 0;
 
 	//~ Begin FSwitchboardTask interface
 	virtual uint32 GetEquivalenceHash() const override
