Switchboard: introduce a CONNECTING status to Device

This adds a CONNECTING value to the DeviceStatus enum and the
accompanying UI handling to represent Devices that are in the process
of establishing a connection.

No Device makes use of this new status yet, but a subsequent change
will update ListenerClient to create its connection on a background thread,
and will modify DeviceUnreal to setup the Qt signal handling for the status
transitions from DISCONNECTED to CONNECTING to CLOSED.

diff --git a/Engine/Plugins/VirtualProduction/Switchboard/Source/Switchboard/switchboard/devices/device_base.py b/Engine/Plugins/VirtualProduction/Switchboard/Source/Switchboard/switchboard/devices/device_base.py
--- a/Engine/Plugins/VirtualProduction/Switchboard/Source/Switchboard/switchboard/devices/device_base.py
+++ b/Engine/Plugins/VirtualProduction/Switchboard/Source/Switchboard/switchboard/devices/device_base.py
@@ -15,6 +15,7 @@ import os, random, socket, threading
 class DeviceStatus(IntEnum):
     DELETE = auto()
     DISCONNECTED = auto()
+    CONNECTING = auto()
     CLOSED = auto()
     SYNCING = auto()
     BUILDING = auto()
@@ -190,7 +191,7 @@ class Device(QtCore.QObject):
 
     @property
     def is_disconnected(self):
-        return self.status in {DeviceStatus.DISCONNECTED}
+        return self.status in {DeviceStatus.DISCONNECTED, DeviceStatus.CONNECTING}
 
     @property
     def project_changelist(self):
@@ -216,6 +217,19 @@ class Device(QtCore.QObject):
     def set_take(self, value):
         pass
 
+    @QtCore.Slot()
+    def connecting_listener(self):
+
+        # Ensure this code is run from the main thread
+        if threading.current_thread() is not threading.main_thread():
+            QtCore.QMetaObject.invokeMethod(self, 'connecting_listener', QtCore.Qt.QueuedConnection)
+            return
+
+        # If the device was disconnected, set to connecting.
+        # We can only transition to CONNECTING from DISCONNECTED.
+        if self.status == DeviceStatus.DISCONNECTED:
+            self.status = DeviceStatus.CONNECTING
+
     @QtCore.Slot()
     def connect_listener(self):
 
@@ -224,7 +238,8 @@ class Device(QtCore.QObject):
             QtCore.QMetaObject.invokeMethod(self, 'connect_listener', QtCore.Qt.QueuedConnection)
             return
 
-        # If the device was disconnected, set to closed
+        # If the device was disconnected, set to closed.
+        # We can transition to CLOSED from either DISCONNECTED or CONNECTING.
         if self.is_disconnected:
             self.status = DeviceStatus.CLOSED
 
diff --git a/Engine/Plugins/VirtualProduction/Switchboard/Source/Switchboard/switchboard/devices/device_widget_base.py b/Engine/Plugins/VirtualProduction/Switchboard/Source/Switchboard/switchboard/devices/device_widget_base.py
--- a/Engine/Plugins/VirtualProduction/Switchboard/Source/Switchboard/switchboard/devices/device_widget_base.py
+++ b/Engine/Plugins/VirtualProduction/Switchboard/Source/Switchboard/switchboard/devices/device_widget_base.py
@@ -185,6 +185,10 @@ class DeviceWidget(QtWidgets.QWidget):
             pixmap = QtGui.QPixmap(f":/icons/images/status_blank_disabled.png")
             self.status_icon.setPixmap(pixmap)
             self.status_icon.setToolTip("Disconnected")
+        elif status == DeviceStatus.CONNECTING:
+            pixmap = QtGui.QPixmap(f":/icons/images/status_orange.png")
+            self.status_icon.setPixmap(pixmap)
+            self.status_icon.setToolTip("Connecting...")
         elif status == DeviceStatus.OPEN:
             pixmap = QtGui.QPixmap(f":/icons/images/status_orange.png")
             self.status_icon.setPixmap(pixmap)
@@ -194,7 +198,7 @@ class DeviceWidget(QtWidgets.QWidget):
             self.status_icon.setToolTip("Connected")
 
         # Device icon
-        if status == DeviceStatus.DISCONNECTED:
+        if status in {DeviceStatus.DISCONNECTED, DeviceStatus.CONNECTING}:
             for label in [self.name_line_edit, self.ip_address_line_edit]:
                 label.setProperty("disconnected", True)
                 label.setStyle(label.style())
@@ -202,10 +206,16 @@ class DeviceWidget(QtWidgets.QWidget):
             pixmap = self.icon_for_state("disabled").pixmap(QtCore.QSize(40, 40))
             self.device_icon.setPixmap(pixmap)
 
-            # Make the Name and IP editable
-            self.name_line_edit.setReadOnly(False)
-            self.ip_address_line_edit.setReadOnly(False)
-        elif previous_status == DeviceStatus.DISCONNECTED and status > DeviceStatus.DISCONNECTED:
+            if status == DeviceStatus.DISCONNECTED:
+                # Make the Name and IP editable when disconnected.
+                self.name_line_edit.setReadOnly(False)
+                self.ip_address_line_edit.setReadOnly(False)
+            elif status == DeviceStatus.CONNECTING:
+                # Make the Name and IP non-editable while connecting.
+                self.name_line_edit.setReadOnly(True)
+                self.ip_address_line_edit.setReadOnly(True)
+        elif ((previous_status in {DeviceStatus.DISCONNECTED, DeviceStatus.CONNECTING}) and
+                status > DeviceStatus.CONNECTING):
             for label in [self.name_line_edit, self.ip_address_line_edit]:
                 label.setProperty("disconnected", False)
                 label.setStyle(label.style())
@@ -213,7 +223,7 @@ class DeviceWidget(QtWidgets.QWidget):
             pixmap = self.icon_for_state("enabled").pixmap(QtCore.QSize(40, 40))
             self.device_icon.setPixmap(pixmap)
 
-            # Make the Name and IP editable
+            # Make the Name and IP non-editable when connected.
             self.name_line_edit.setReadOnly(True)
             self.ip_address_line_edit.setReadOnly(True)
 
diff --git a/Engine/Plugins/VirtualProduction/Switchboard/Source/Switchboard/switchboard/devices/kipro/plugin_kipro.py b/Engine/Plugins/VirtualProduction/Switchboard/Source/Switchboard/switchboard/devices/kipro/plugin_kipro.py
--- a/Engine/Plugins/VirtualProduction/Switchboard/Source/Switchboard/switchboard/devices/kipro/plugin_kipro.py
+++ b/Engine/Plugins/VirtualProduction/Switchboard/Source/Switchboard/switchboard/devices/kipro/plugin_kipro.py
@@ -269,7 +269,7 @@ class DeviceWidgetKiPro(DeviceWidget):
     def update_status(self, status, previous_status):
         super().update_status(status, previous_status)
 
-        if status == DeviceStatus.DISCONNECTED:
+        if status in {DeviceStatus.DISCONNECTED, DeviceStatus.CONNECTING}:
             self.play_button.setDisabled(True)
             self.stop_button.setDisabled(True)
         elif status >= DeviceStatus.CLOSED:
diff --git a/Engine/Plugins/VirtualProduction/Switchboard/Source/Switchboard/switchboard/devices/unreal/plugin_unreal.py b/Engine/Plugins/VirtualProduction/Switchboard/Source/Switchboard/switchboard/devices/unreal/plugin_unreal.py
--- a/Engine/Plugins/VirtualProduction/Switchboard/Source/Switchboard/switchboard/devices/unreal/plugin_unreal.py
+++ b/Engine/Plugins/VirtualProduction/Switchboard/Source/Switchboard/switchboard/devices/unreal/plugin_unreal.py
@@ -1278,13 +1278,7 @@ class DeviceWidgetUnreal(DeviceWidget):
         self.signal_device_widget_close.emit(self)
 
     def _connect(self):
-        # Make sure the button is in the correct state
-        self.connect_button.setChecked(True)
-        self.connect_button.setToolTip("Disconnect from listener")
-
-        self.open_button.setDisabled(False)
-        self.sync_button.setDisabled(False)
-        self.build_button.setDisabled(False)
+        self._update_connected_ui()
 
         # Emit Signal to Switchboard
         self.signal_device_widget_connect.emit(self)
@@ -1312,11 +1306,30 @@ class DeviceWidgetUnreal(DeviceWidget):
         self.sync_button.setDisabled(True)
         self.build_button.setDisabled(True)
 
+    def _update_connected_ui(self):
+        ''' Updates the UI of the device to reflect connected status. '''
+        # Make sure the button is in the correct state
+        self.connect_button.setChecked(True)
+        self.connect_button.setToolTip("Disconnect from listener")
+
+        self.open_button.setDisabled(False)
+        self.sync_button.setDisabled(False)
+        self.build_button.setDisabled(False)
+
     def update_status(self, status, previous_status):
         super().update_status(status, previous_status)
 
-        if status == DeviceStatus.DISCONNECTED:
+        if status <= DeviceStatus.CONNECTING:
             self._update_disconnected_ui()
+        else:
+            self._update_connected_ui()
+
+        # The connect/disconnect button is enabled in all states except for CONNECTING.
+        self.connect_button.setDisabled(False)
+
+        if status == DeviceStatus.CONNECTING:
+            self.connect_button.setDisabled(True)
+            self.connect_button.setToolTip("Connecting to listener...")
         elif status == DeviceStatus.CLOSED:
             self.open_button.setDisabled(False)
             self.open_button.setChecked(False)
