Switchboard: enable connecting ListenerClients on background threads

A new (default) mode was added to ListenerClient's connect() function that
establishes the connection on a background thread and emits Qt signals through
its handler that indicate whether the connection succeeded or failed. The
original blocking connection behavior can still be invoked by passing the
blocking=True parameter to connect().

DeviceUnreal is the only device that makes use of ListenerClient, so it was
updated to connect asynchronously and listen for connection state via Qt
signals.

This addresses an issue where creating connections to Unreal devices blocks the
main thread and locks up the UI. When connections succeed, the blocking is hardly
noticeable, but in cases where connections may fail (for example when
SwitchboardListener is not running on the device), the UI will be locked up until the
connection timeout expires (roughly 20 seconds). This gets compounded when there
are multiple Unreal devices setup that don't have their listeners running, as the
connections were being attempted serially.

With these changes, creating the connection no longer blocks the main thread and
multiple connections can be created in parallel.

The only other use of ListenerClient is in the redeploy dialog, which was modified to
use the blocking=True parameter to connect() for now just to preserve its current
behavior.

diff --git a/Engine/Plugins/VirtualProduction/Switchboard/Source/Switchboard/switchboard/devices/unreal/plugin_unreal.py b/Engine/Plugins/VirtualProduction/Switchboard/Source/Switchboard/switchboard/devices/unreal/plugin_unreal.py
--- a/Engine/Plugins/VirtualProduction/Switchboard/Source/Switchboard/switchboard/devices/unreal/plugin_unreal.py
+++ b/Engine/Plugins/VirtualProduction/Switchboard/Source/Switchboard/switchboard/devices/unreal/plugin_unreal.py
@@ -373,6 +373,10 @@ class DeviceUnreal(Device):
         self.inflight_project_cl = None
         self.inflight_engine_cl = None
 
+        self.unreal_client.listener_qt_handler.listener_connecting.connect(super().connecting_listener, QtCore.Qt.QueuedConnection)
+        self.unreal_client.listener_qt_handler.listener_connected.connect(super().connect_listener, QtCore.Qt.QueuedConnection)
+        self.unreal_client.listener_qt_handler.listener_connection_failed.connect(self._on_listener_connection_failed, QtCore.Qt.QueuedConnection)
+
         # Set a delegate method if the device gets a disconnect signal
         self.unreal_client.disconnect_delegate = self.on_listener_disconnect
         self.unreal_client.receive_file_completed_delegate = self.on_file_received
@@ -568,13 +572,14 @@ class DeviceUnreal(Device):
             QtCore.QMetaObject.invokeMethod(self, 'connect_listener', QtCore.Qt.QueuedConnection)
             return
 
-        is_connected = self.unreal_client.connect()
-
-        if not is_connected:
-            self.device_qt_handler.signal_device_connect_failed.emit(self)
-            return
+        # This will start the connection process asynchronously. The base class will be notified
+        # by signal whether the connection succeeded or failed and will update the device's status
+        # accordingly.
+        self.unreal_client.connect()
 
-        super().connect_listener()
+    @QtCore.Slot()
+    def _on_listener_connection_failed(self):
+        self.device_qt_handler.signal_device_connect_failed.emit(self)
 
     @QtCore.Slot()
     def disconnect_listener(self):
diff --git a/Engine/Plugins/VirtualProduction/Switchboard/Source/Switchboard/switchboard/devices/unreal/redeploy_dialog.py b/Engine/Plugins/VirtualProduction/Switchboard/Source/Switchboard/switchboard/devices/unreal/redeploy_dialog.py
--- a/Engine/Plugins/VirtualProduction/Switchboard/Source/Switchboard/switchboard/devices/unreal/redeploy_dialog.py
+++ b/Engine/Plugins/VirtualProduction/Switchboard/Source/Switchboard/switchboard/devices/unreal/redeploy_dialog.py
@@ -105,7 +105,7 @@ class RedeployListenerEndpoint(QtCore.QObject):
         # Suppress on_disconnect for this expected disconnect
         self.client.disconnect()
         self.client.disconnect_delegate = self.on_disconnect
-        self.client.connect()
+        self.client.connect(blocking=True)
 
     def on_disconnect(self, unexpected, exception):
         self.version = None
@@ -190,7 +190,7 @@ class RedeployListenerDialog(QtWidgets.QDialog):
                 endpoint = RedeployListenerEndpoint(self, device.unreal_client.ip_address, device.unreal_client.port)
                 endpoint.signal_result.connect(self.on_endpoint_result)
                 endpoint.signal_refresh_ui.connect(self.refresh_ui)
-                endpoint.client.connect()
+                endpoint.client.connect(blocking=True)
                 layout.addWidget(endpoint.endpoint_label, row_idx, 0)
                 layout.addWidget(endpoint.devices_label,  row_idx, 1)
                 layout.addWidget(endpoint.version_label,  row_idx, 2)
diff --git a/Engine/Plugins/VirtualProduction/Switchboard/Source/Switchboard/switchboard/listener_client.py b/Engine/Plugins/VirtualProduction/Switchboard/Source/Switchboard/switchboard/listener_client.py
--- a/Engine/Plugins/VirtualProduction/Switchboard/Source/Switchboard/switchboard/listener_client.py
+++ b/Engine/Plugins/VirtualProduction/Switchboard/Source/Switchboard/switchboard/listener_client.py
@@ -2,12 +2,20 @@
 from . import message_protocol
 from .switchboard_logging import LOGGER
 
+from PySide2 import QtCore
+
 import datetime, select, socket, uuid, traceback, typing
 
 from collections import deque
 from threading import Thread
 
 
+class ListenerQtHandler(QtCore.QObject):
+    listener_connecting = QtCore.Signal(object)
+    listener_connected = QtCore.Signal(object)
+    listener_connection_failed = QtCore.Signal(object)
+
+
 class ListenerClient(object):
     '''Connects to a server running SwitchboardListener.
 
@@ -28,6 +36,8 @@ class ListenerClient(object):
         self.port = port
         self.buffer_size = buffer_size
 
+        self.listener_qt_handler = ListenerQtHandler()
+
         self.message_queue = deque()
         self.close_socket = False
 
@@ -78,18 +88,42 @@ class ListenerClient(object):
             return False
         return False
 
-    def connect(self, ip_address=None):
+    def connect(self, ip_address=None, blocking=False):
+        '''
+        Initiates a connection.
+
+        By default, calling this function starts the connection setup on a background
+        thread. The caller should connect slots to the listener_qt_handler's
+        "listener_connected" and "listener_connection_failed" signals to determine
+        whether the connection was created successfully.
+
+        When "blocking" is True, this function will not return until establishing the
+        connection either succeeds or fails, and True or False will be returned, respectively.
+        '''
         self.disconnect()
 
         if ip_address:
             self.ip_address = ip_address
         elif not self.ip_address:
             LOGGER.debug('No ip_address has been set. Cannot connect')
+            self.listener_qt_handler.listener_connection_failed.emit(self)
             return False
 
         self.close_socket = False
         self.last_activity = datetime.datetime.now()
 
+        if blocking:
+            # Since we're blocking while creating the connection, we don't pass through
+            # the CONNECTING status.
+            return self._establish_connection()
+
+        self.listener_qt_handler.listener_connecting.emit(self)
+        establish_connection_thread = Thread(target=self._establish_connection)
+        establish_connection_thread.start()
+
+        return True
+
+    def _establish_connection(self):
         try:
             LOGGER.info(f"Connecting to {self.ip_address}:{self.port}")
 
@@ -103,8 +137,10 @@ class ListenerClient(object):
         except OSError:
             LOGGER.error(f"Socket error: {self.ip_address}:{self.port}")
             self.socket = None
+            self.listener_qt_handler.listener_connection_failed.emit(self)
             return False
 
+        self.listener_qt_handler.listener_connected.emit(self)
         return True
 
     def disconnect(self, unexpected=False, exception=None):
