Switchboard: Replace "Install Listener" and "Uninstall Listener" toolbar menu items with "Launch Switchboard Listener on Login" checkbox. Fixed registry autorun entry not working because it was added with forward slashes in the path. Automatically update registry autorun entry when "Listener Commandline Arguments" setting changes.

This commit removes install_listener.py and moves the corresponding functionality into the SwitchboardEditor module.

diff --git a/Engine/Plugins/VirtualProduction/Switchboard/Source/SwitchboardEditor/Private/SwitchboardEditorModule.cpp b/Engine/Plugins/VirtualProduction/Switchboard/Source/SwitchboardEditor/Private/SwitchboardEditorModule.cpp
--- a/Engine/Plugins/VirtualProduction/Switchboard/Source/SwitchboardEditor/Private/SwitchboardEditorModule.cpp
+++ b/Engine/Plugins/VirtualProduction/Switchboard/Source/SwitchboardEditor/Private/SwitchboardEditorModule.cpp
@@ -1,77 +1,271 @@
 // Copyright Epic Games, Inc. All Rights Reserved.
 
-#include "CoreMinimal.h"
+#include "SwitchboardEditorModule.h"
+
 #include "Editor.h"
-#include "Editor/Blutility/Classes/EditorUtilityObject.h"
 #include "EditorUtilitySubsystem.h"
 #include "ISettingsModule.h"
+#include "ISettingsSection.h"
 #include "Misc/CoreDelegates.h"
-#include "Modules/ModuleInterface.h"
-#include "Modules/ModuleManager.h"
+#include "Misc/ScopeExit.h"
 #include "SwitchboardEditorSettings.h"
 #include "SwitchboardMenuEntry.h"
-#include "Misc/CoreDelegates.h"
-#include "Editor.h"
 
+#if PLATFORM_WINDOWS
+#include "Windows/AllowWindowsPlatformTypes.h"
+#include <winreg.h>
+#include "Windows/HideWindowsPlatformTypes.h"
+
+static const FString WindowsRunRegKeyPath = "Software\\Microsoft\\Windows\\CurrentVersion\\Run";
+static const FString ListenerAutolaunchRegValueName = "SwitchboardListener";
+
+FString GetListenerAutolaunchEntry();
+FString GetListenerAutolaunchEntryExecutable();
+bool SetListenerAutolaunchEntry(const FString& NewCommandLine);
+bool RemoveListenerAutolaunchEntry();
+#endif
 
 #define LOCTEXT_NAMESPACE "SwitchboardEditorModule"
 
+DEFINE_LOG_CATEGORY(LogSwitchboardPlugin);
 
-FDelegateHandle DeferredStartDelegateHandle;
 
-class FSwitchboardEditorModule : public IModuleInterface
+void FSwitchboardEditorModule::StartupModule()
 {
-public:
-	virtual void StartupModule() override
+	FSwitchboardMenuEntry::Register();
+
+	if (ISettingsModule* SettingsModule = FModuleManager::GetModulePtr<ISettingsModule>("Settings"))
 	{
-		FSwitchboardMenuEntry::Register();
+		ISettingsSectionPtr SettingsSection = SettingsModule->RegisterSettings("Project", "Plugins", "Switchboard",
+			LOCTEXT("SettingsName", "Switchboard"),
+			LOCTEXT("Description", "Configure the Switchboard launcher."),
+			GetMutableDefault<USwitchboardEditorSettings>()
+		);
 
-		if (ISettingsModule* SettingsModule = FModuleManager::GetModulePtr<ISettingsModule>("Settings"))
+		SettingsSection->OnModified().BindRaw(this, &FSwitchboardEditorModule::OnSettingsModified);
+	}
+
+	DeferredStartDelegateHandle = FCoreDelegates::OnFEngineLoopInitComplete.AddRaw(this, &FSwitchboardEditorModule::OnEngineInitComplete);
+}
+
+
+void FSwitchboardEditorModule::ShutdownModule()
+{
+	if (!IsRunningCommandlet() && UObjectInitialized() && !IsEngineExitRequested())
+	{
+		ISettingsModule* SettingsModule = FModuleManager::GetModulePtr<ISettingsModule>("Settings");
+		if (SettingsModule != nullptr)
 		{
-			SettingsModule->RegisterSettings("Project", "Plugins", "Switchboard",
-				LOCTEXT("SettingsName", "Switchboard"),
-				LOCTEXT("Description", "Configure the Switchboard launcher."),
-				GetMutableDefault<USwitchboardEditorSettings>()
-			);
+			SettingsModule->UnregisterSettings("Project", "Plugins", "Switchboard");
 		}
+		FSwitchboardMenuEntry::Unregister();
+	}
+}
+
+
+#if PLATFORM_WINDOWS
+bool FSwitchboardEditorModule::IsListenerAutolaunchEnabled() const
+{
+	const FString& ExistingCmd = GetListenerAutolaunchEntry().TrimEnd();
+	const FString& ConfigCmd = GetDefault<USwitchboardEditorSettings>()->GetListenerInvocation().TrimEnd();
+	return ExistingCmd == ConfigCmd;
+}
+
 
-		DeferredStartDelegateHandle = FCoreDelegates::OnFEngineLoopInitComplete.AddRaw(this, &FSwitchboardEditorModule::RunDefaultOSCListener);
+bool FSwitchboardEditorModule::SetListenerAutolaunchEnabled(bool bEnabled)
+{
+	if (bEnabled)
+	{
+		const FString CommandLine = GetDefault<USwitchboardEditorSettings>()->GetListenerInvocation();
+		return SetListenerAutolaunchEntry(CommandLine);
+	}
+	else
+	{
+		return RemoveListenerAutolaunchEntry();
 	}
+}
+#endif // #if PLATFORM_WINDOWS
 
-	virtual void ShutdownModule() override
+
+void FSwitchboardEditorModule::OnEngineInitComplete()
+{
+	FCoreDelegates::OnFEngineLoopInitComplete.Remove(DeferredStartDelegateHandle);
+	DeferredStartDelegateHandle.Reset();
+
+	RunDefaultOSCListener();
+}
+
+
+bool FSwitchboardEditorModule::OnSettingsModified()
+{
+#if PLATFORM_WINDOWS
+	// If the existing entry's listener executable path matches our current engine / path,
+	// then we ensure the command line arguments also stay in sync with our settings.
+	const FString AutolaunchEntryExecutable = GetListenerAutolaunchEntryExecutable();
+	if (!AutolaunchEntryExecutable.IsEmpty())
 	{
-		if (!IsRunningCommandlet() && UObjectInitialized() && !IsEngineExitRequested())
+		const FString ConfigExecutable = GetDefault<USwitchboardEditorSettings>()->GetListenerPlatformPath();
+		if (AutolaunchEntryExecutable == ConfigExecutable)
+		{
+			UE_LOG(LogSwitchboardPlugin, Log, TEXT("Updating listener auto-launch entry"));
+			SetListenerAutolaunchEnabled(true);
+		}
+		else
 		{
-			ISettingsModule* SettingsModule = FModuleManager::GetModulePtr<ISettingsModule>("Settings");
-			if (SettingsModule != nullptr)
-			{
-				SettingsModule->UnregisterSettings("Editor", "Plugins", "Switchboard");
-			}
-			FSwitchboardMenuEntry::Unregister();
+			UE_LOG(LogSwitchboardPlugin, Log, TEXT("NOT updating listener auto-launch; paths differ:\n\t%s\n\t%s"),
+				*AutolaunchEntryExecutable, *ConfigExecutable);
 		}
 	}
+#endif // #if PLATFORM_WINDOWS
+
+	return true;
+}
 
-	void RunDefaultOSCListener()
-	{
-		if (DeferredStartDelegateHandle.IsValid())
-		{
-			FCoreDelegates::OnEndFrame.Remove(DeferredStartDelegateHandle);
-			DeferredStartDelegateHandle.Reset();
-		};
 
-		USwitchboardEditorSettings* SwitchboardEditorSettings = USwitchboardEditorSettings::GetSwitchboardEditorSettings();
-		FSoftObjectPath SwitchboardOSCListener = SwitchboardEditorSettings->SwitchboardOSCListener;
-		if (SwitchboardOSCListener.IsValid())
+void FSwitchboardEditorModule::RunDefaultOSCListener()
+{
+	USwitchboardEditorSettings* SwitchboardEditorSettings = USwitchboardEditorSettings::GetSwitchboardEditorSettings();
+	FSoftObjectPath SwitchboardOSCListener = SwitchboardEditorSettings->SwitchboardOSCListener;
+	if (SwitchboardOSCListener.IsValid())
+	{
+		UObject* SwitchboardOSCListenerObject = SwitchboardOSCListener.TryLoad();
+		if (SwitchboardOSCListenerObject && !SwitchboardOSCListenerObject->IsPendingKillOrUnreachable() && GEditor)
 		{
-			UObject* SwitchboardOSCListenerObject = SwitchboardOSCListener.TryLoad();
-			if (SwitchboardOSCListenerObject && !SwitchboardOSCListenerObject->IsPendingKillOrUnreachable() && GEditor)
-			{
-				GEditor->GetEditorSubsystem<UEditorUtilitySubsystem>()->TryRun(SwitchboardOSCListenerObject);
-			}
+			GEditor->GetEditorSubsystem<UEditorUtilitySubsystem>()->TryRun(SwitchboardOSCListenerObject);
 		}
 	}
+}
+
+
+#if PLATFORM_WINDOWS
+HKEY OpenHkcuRunKey()
+{
+	HKEY HkcuRunKey = nullptr;
+	const LSTATUS OpenResult = RegOpenKeyEx(HKEY_CURRENT_USER, *WindowsRunRegKeyPath, 0, KEY_ALL_ACCESS, &HkcuRunKey);
+	if (OpenResult != ERROR_SUCCESS)
+	{
+		UE_LOG(LogSwitchboardPlugin, Error, TEXT("Error opening registry key %s (%08X)"), *WindowsRunRegKeyPath, OpenResult);
+		return nullptr;
+	}
+	return HkcuRunKey;
+}
+
+
+FString GetListenerAutolaunchEntry()
+{
+	HKEY HkcuRunKey = OpenHkcuRunKey();
+	if (!HkcuRunKey)
+	{
+		return FString();
+	}
+
+	ON_SCOPE_EXIT
+	{
+		RegCloseKey(HkcuRunKey);
+	};
+
+	DWORD ValueType;
+	DWORD ValueSizeBytes = 0;
+	const LSTATUS SizeResult = RegQueryValueEx(HkcuRunKey, *ListenerAutolaunchRegValueName, nullptr, &ValueType, nullptr, &ValueSizeBytes);
+	if (SizeResult == ERROR_FILE_NOT_FOUND)
+	{
+		return FString();
+	}
+	else if (SizeResult != ERROR_SUCCESS)
+	{
+		UE_LOG(LogSwitchboardPlugin, Error, TEXT("Error reading registry value %s:\"%s\" (%08X)"), *WindowsRunRegKeyPath, *ListenerAutolaunchRegValueName, SizeResult);
+		return FString();
+	}
+
+	FString Value;
+	TArray<TCHAR>& CharArray = Value.GetCharArray();
+	const uint32 ValueLenChars = ValueSizeBytes / sizeof(TCHAR);
+	CharArray.SetNumUninitialized(ValueLenChars);
+
+	const LSTATUS QueryResult = RegQueryValueEx(HkcuRunKey, *ListenerAutolaunchRegValueName, nullptr, &ValueType, reinterpret_cast<LPBYTE>(CharArray.GetData()), &ValueSizeBytes);
+	if (QueryResult != ERROR_SUCCESS)
+	{
+		UE_LOG(LogSwitchboardPlugin, Error, TEXT("Error reading registry value %s:\"%s\" (%08X)"), *WindowsRunRegKeyPath, *ListenerAutolaunchRegValueName, QueryResult);
+		return FString();
+	}
+	else if (ValueType != REG_SZ)
+	{
+		UE_LOG(LogSwitchboardPlugin, Error, TEXT("Registry value %s:\"%s\" has wrong type (%u)"), *WindowsRunRegKeyPath, *ListenerAutolaunchRegValueName, ValueType);
+		return FString();
+	}
+
+	if (CharArray[CharArray.Num() - 1] != TEXT('\0'))
+	{
+		CharArray.Add(TEXT('\0'));
+	}
+
+	return Value;
+}
+
+
+FString GetListenerAutolaunchEntryExecutable()
+{
+	FString AutolaunchCommand = GetListenerAutolaunchEntry();
+	if (AutolaunchCommand.IsEmpty())
+	{
+		return FString();
+	}
+
+	TArray<FString> QuoteParts;
+	AutolaunchCommand.ParseIntoArray(QuoteParts, TEXT("\""));
+	return QuoteParts[0];
+}
+
+
+bool SetListenerAutolaunchEntry(const FString& NewCommandLine)
+{
+	HKEY HkcuRunKey = OpenHkcuRunKey();
+	if (!HkcuRunKey)
+	{
+		return false;
+	}
+
+	ON_SCOPE_EXIT
+	{
+		RegCloseKey(HkcuRunKey);
+	};
+
+	const TArray<TCHAR>& CharArray = NewCommandLine.GetCharArray();
+	const LSTATUS SetResult = RegSetValueEx(HkcuRunKey, *ListenerAutolaunchRegValueName, 0, REG_SZ, reinterpret_cast<const BYTE*>(CharArray.GetData()), CharArray.Num() * sizeof(TCHAR));
+	if (SetResult != ERROR_SUCCESS)
+	{
+		UE_LOG(LogSwitchboardPlugin, Error, TEXT("Error setting registry value %s:\"%s\" (%08X)"), *WindowsRunRegKeyPath, *ListenerAutolaunchRegValueName, SetResult);
+		return false;
+	}
+
+	return true;
+}
+
+
+bool RemoveListenerAutolaunchEntry()
+{
+	HKEY HkcuRunKey = OpenHkcuRunKey();
+	if (!HkcuRunKey)
+	{
+		return true;
+	}
+
+	ON_SCOPE_EXIT
+	{
+		RegCloseKey(HkcuRunKey);
+	};
+
+	const LSTATUS DeleteResult = RegDeleteValue(HkcuRunKey, *ListenerAutolaunchRegValueName);
+	if (DeleteResult != ERROR_SUCCESS)
+	{
+		UE_LOG(LogSwitchboardPlugin, Error, TEXT("Error deleting registry value %s:\"%s\" (%08X)"), *WindowsRunRegKeyPath, *ListenerAutolaunchRegValueName, DeleteResult);
+		return false;
+	}
+
+	return true;
+}
+#endif // #if PLATFORM_WINDOWS
 
-};
 
 IMPLEMENT_MODULE(FSwitchboardEditorModule, SwitchboardEditor);
 
diff --git a/Engine/Plugins/VirtualProduction/Switchboard/Source/SwitchboardEditor/Private/SwitchboardEditorModule.h b/Engine/Plugins/VirtualProduction/Switchboard/Source/SwitchboardEditor/Private/SwitchboardEditorModule.h
new file mode 100644
--- /dev/null
+++ b/Engine/Plugins/VirtualProduction/Switchboard/Source/SwitchboardEditor/Private/SwitchboardEditorModule.h
@@ -0,0 +1,50 @@
+// Copyright Epic Games, Inc. All Rights Reserved.
+
+#pragma once
+
+#include "CoreMinimal.h"
+#include "Delegates/IDelegateInstance.h"
+#include "Modules/ModuleInterface.h"
+#include "Modules/ModuleManager.h"
+
+
+DECLARE_LOG_CATEGORY_EXTERN(LogSwitchboardPlugin, Log, All);
+
+
+class FSwitchboardEditorModule : public IModuleInterface
+{
+public:
+	/**
+	 * Singleton-like access to this module's interface.  This is just for convenience!
+	 * Beware of calling this during the shutdown phase, though.  Your module might have been unloaded already.
+	 *
+	 * @return Returns singleton instance, loading the module on demand if needed
+	 */
+	static FSwitchboardEditorModule& Get()
+	{
+		static const FName ModuleName = "SwitchboardEditor";
+		return FModuleManager::LoadModuleChecked<FSwitchboardEditorModule>(ModuleName);
+	};
+
+	//~ Begin IModuleInterface
+	virtual void StartupModule() override;
+	virtual void ShutdownModule() override;
+	//~ End IModuleInterface
+
+#if PLATFORM_WINDOWS
+	/** Returns whether (this engine's) SwitchboardListener is configured to run automatically. */
+	bool IsListenerAutolaunchEnabled() const;
+
+	/** Enables or disables auto-run of SwitchboardListener. */
+	bool SetListenerAutolaunchEnabled(bool bEnabled);
+#endif
+
+private:
+	void OnEngineInitComplete();
+	bool OnSettingsModified();
+
+	void RunDefaultOSCListener();
+
+private:
+	FDelegateHandle DeferredStartDelegateHandle;
+};
diff --git a/Engine/Plugins/VirtualProduction/Switchboard/Source/SwitchboardEditor/Private/SwitchboardMenuEntry.cpp b/Engine/Plugins/VirtualProduction/Switchboard/Source/SwitchboardEditor/Private/SwitchboardMenuEntry.cpp
--- a/Engine/Plugins/VirtualProduction/Switchboard/Source/SwitchboardEditor/Private/SwitchboardMenuEntry.cpp
+++ b/Engine/Plugins/VirtualProduction/Switchboard/Source/SwitchboardEditor/Private/SwitchboardMenuEntry.cpp
@@ -12,13 +12,12 @@
 #include "Misc/FileHelper.h"
 #include "Misc/Paths.h"
 #include "Modules/ModuleManager.h"
+#include "SwitchboardEditorModule.h"
 #include "SwitchboardEditorSettings.h"
 #include "SwitchboardEditorStyle.h"
 
 #define LOCTEXT_NAMESPACE "SwitchboardEditor"
 
-DEFINE_LOG_CATEGORY(LogSwitchboardPlugin);
-
 
 class FSwitchboardUICommands : public TCommands<FSwitchboardUICommands>
 {
@@ -75,9 +74,6 @@ struct FSwitchboardMenuEntryImpl
 
 		FLevelEditorModule& LevelEditorModule = FModuleManager::LoadModuleChecked<FLevelEditorModule>("LevelEditor");
 		LevelEditorModule.GetToolBarExtensibilityManager()->AddExtender(ToolBarExtender);
-	
-		PathToListenerInstaller = FPaths::ConvertRelativePathToFull(FPaths::EnginePluginsDir() + FString(TEXT("VirtualProduction")));
-		PathToListenerInstaller /= FString(TEXT("Switchboard")) / FString(TEXT("Source")) / FString(TEXT("install_listener.py"));
 	}
 
 	~FSwitchboardMenuEntryImpl()
@@ -105,19 +101,25 @@ struct FSwitchboardMenuEntryImpl
 				FSlateIcon(),
 				FUIAction(FExecuteAction::CreateRaw(this, &FSwitchboardMenuEntryImpl::LaunchListener))
 			);
+
 #if PLATFORM_WINDOWS
 			MenuBuilder.AddMenuSeparator();
+
+			// The FGetActionCheckState delegate is called continually while the menu is open.
+			// Cache the state now as the combo is initially expanded.
+			CachedAutolaunchCheckState = GetAutolaunchCheckState();
+
 			MenuBuilder.AddMenuEntry(
-				LOCTEXT("ListenerInstallLabel", "Install SwitchboardListener"),
-				LOCTEXT("ListenerInstallTooltip", "Configures SwitchboardListener to automatically run on system start."),
-				FSlateIcon(),
-				FUIAction(FExecuteAction::CreateRaw(this, &FSwitchboardMenuEntryImpl::InstallListener))
-			);
-			MenuBuilder.AddMenuEntry(
-				LOCTEXT("ListenerUninstallLabel", "Uninstall SwitchboardListener"),
-				LOCTEXT("ListenerUninstallTooltip", "Stop SwitchboardListener from automatically running on system start."),
+				LOCTEXT("ListenerAutolaunchLabel", "Launch Switchboard Listener on Login"),
+				LOCTEXT("ListenerAutolaunchTooltip", "Controls whether SwitchboardListener runs automatically when you log into Windows."),
 				FSlateIcon(),
-				FUIAction(FExecuteAction::CreateRaw(this, &FSwitchboardMenuEntryImpl::UninstallListener))
+				FUIAction(
+					FExecuteAction::CreateRaw(this, &FSwitchboardMenuEntryImpl::ToggleAutolaunch),
+					FCanExecuteAction(),
+					FGetActionCheckState::CreateLambda([&](){ return CachedAutolaunchCheckState; })
+				),
+				NAME_None,
+				EUserInterfaceActionType::ToggleButton
 			);
 #endif
 		}
@@ -148,95 +150,75 @@ struct FSwitchboardMenuEntryImpl
 		}
 	}
 
-	void InstallListener()
+#if PLATFORM_WINDOWS
+	ECheckBoxState GetAutolaunchCheckState()
 	{
-		FString PythonInterpreter = GetPathToPython();
-		if (PythonInterpreter.IsEmpty())
-		{
-			const FString ErrorMsg = TEXT("Could not find path to python interpreter which is required to run the installation script!");
-			UE_LOG(LogSwitchboardPlugin, Error, TEXT("%s"), *ErrorMsg);
-			FPlatformMisc::MessageBoxExt(EAppMsgType::Ok, *ErrorMsg, TEXT("No Python interpreter"));
-			return;
-		}
-
-		const FString ListenerPath = GetDefault<USwitchboardEditorSettings>()->ListenerPath.FilePath;
-		if (!FPaths::FileExists(ListenerPath))
-		{
-			const FString ErrorMsg = TEXT("Could not find SwitchboardListener! Make sure it was compiled.");
-			FPlatformMisc::MessageBoxExt(EAppMsgType::Ok, *ErrorMsg, TEXT("Error installing listener"));
-			return;
-		}
-
-		const FString ListenerArgs = GetDefault<USwitchboardEditorSettings>()->ListenerCommandlineArguments;
-		const FString Args = FString::Printf(TEXT("%s install \"\\\"%s\\\" %s\""), *PathToListenerInstaller, *ListenerPath, *ListenerArgs);
-		if (RunProcess(PythonInterpreter, Args))
-		{
-			const FString Msg = TEXT("Successfully installed the SwitchboardListener. It will automatically run on the next system start.");
-			FPlatformMisc::MessageBoxExt(EAppMsgType::Ok, *Msg, TEXT("Success"));
-		}
-		else
-		{
-			const FString ErrorMsg = TEXT("Unable to install listener! Check the log for details.");
-			FPlatformMisc::MessageBoxExt(EAppMsgType::Ok, *ErrorMsg, TEXT("Error installing listener"));
-		}
+		return FSwitchboardEditorModule::Get().IsListenerAutolaunchEnabled() ? ECheckBoxState::Checked : ECheckBoxState::Unchecked;
 	}
 
-	void UninstallListener()
+	void ToggleAutolaunch()
 	{
-		FString PythonInterpreter = GetPathToPython();
-		if (PythonInterpreter.IsEmpty())
+		if (CachedAutolaunchCheckState == ECheckBoxState::Checked)
 		{
-			return;
-		}
-
-		const FString Args = FString::Printf(TEXT("%s uninstall"), *PathToListenerInstaller);
-		if (RunProcess(PythonInterpreter, Args))
-		{
-			const FString Msg = TEXT("Successfully removed the SwitchboardListener from Autorun.");
-			FPlatformMisc::MessageBoxExt(EAppMsgType::Ok, *Msg, TEXT("Success"));
+			FSwitchboardEditorModule::Get().SetListenerAutolaunchEnabled(false);
 		}
 		else
 		{
-			const FString ErrorMsg = TEXT("Unable to uninstall listener! Check the log for details.");
-			FPlatformMisc::MessageBoxExt(EAppMsgType::Ok, *ErrorMsg, TEXT("Error uninstalling listener"));
+			if (!FPaths::FileExists(GetDefault<USwitchboardEditorSettings>()->ListenerPath.FilePath))
+			{
+				const FString ErrorMsg = TEXT("Could not find SwitchboardListener! Make sure it has been compiled.");
+				FPlatformMisc::MessageBoxExt(EAppMsgType::Ok, *ErrorMsg, TEXT("Error enabling SwitchboardListener auto-launch"));
+				return;
+			}
+
+			FSwitchboardEditorModule::Get().SetListenerAutolaunchEnabled(true);
 		}
+
+		CachedAutolaunchCheckState = GetAutolaunchCheckState();
 	}
+#endif // #if PLATFORM_WINDOWS
 
 	FString GetPathToPython()
 	{
-		const FFilePath PythonPath = GetDefault<USwitchboardEditorSettings>()->PythonInterpreterPath;
-		FString PythonInterpreter = PythonPath.FilePath;
-		if (PythonInterpreter.IsEmpty())
+		// First try the user-specified path from settings.
+		const FString& UserInterpreter = GetDefault<USwitchboardEditorSettings>()->PythonInterpreterPath.FilePath;
+		if (!UserInterpreter.IsEmpty() && FPaths::FileExists(UserInterpreter))
 		{
-			FString PythonExeName;
+			return UserInterpreter;
+		}
+
+		const FString PythonExeName =
 #if PLATFORM_WINDOWS
-			PythonExeName = TEXT("python.exe");
-#elif PLATFORM_LINUX
-			PythonExeName = TEXT("python");
+			TEXT("python.exe");
+#else
+			TEXT("python");
 #endif
 
-			PythonInterpreter = FPaths::EngineDir() / TEXT("Extras") / TEXT("ThirdPartyNotUE") /
-				TEXT("SwitchboardThirdParty") / TEXT("Python") / TEXT("current") / PythonExeName;
-			if (!FPaths::FileExists(PythonInterpreter)) // python will not exist in ThirdPartyNotUE (yet) if SB was never run
-			{
-				// so we fall back onto the interpreter that ships with the engine
-				PythonInterpreter = FPaths::ConvertRelativePathToFull(FPaths::EngineDir() + FString(TEXT("Binaries")));
-				PythonInterpreter /= FString(TEXT("ThirdParty")) / FString(TEXT("Python3")) / PythonExeName;
-			}
+		// If it's been installed, fall back to Switchboard's embedded Python distribution.
+		const FString SwitchboardInterpreter = FPaths::EngineDir() / TEXT("Extras") / TEXT("ThirdPartyNotUE") /
+			TEXT("SwitchboardThirdParty") / TEXT("Python") / TEXT("current") / PythonExeName;
+		if (FPaths::FileExists(SwitchboardInterpreter))
+		{
+			return SwitchboardInterpreter;
 		}
 
-		if (!FPaths::FileExists(PythonInterpreter))
-		{
-			FString EnginePythonPath = FPaths::EngineDir() / TEXT("Binaries") / TEXT("ThirdParty") / TEXT("Python3");
-#if PLATFORM_WINDOWS
-			EnginePythonPath = EnginePythonPath / TEXT("Win64") / TEXT("python.exe");
-#elif PLATFORM_LINUX
-			EnginePythonPath = EnginePythonPath / TEXT("Linux") / TEXT("bin") / TEXT("python");
+		// Finally, fall back to the engine ThirdParty Python.
+		const FString EngineInterpreter = FPaths::EngineDir() / TEXT("Binaries") / TEXT("ThirdParty")
+			/ TEXT("Python3") / FPlatformProcess::GetBinariesSubdirectory()
+#if !PLATFORM_WINDOWS
+			/ TEXT("bin")
 #endif
-			PythonInterpreter = EnginePythonPath;
+			/ PythonExeName;
+		if (FPaths::FileExists(EngineInterpreter))
+		{
+			return EngineInterpreter;
 		}
 
-		return PythonInterpreter;
+		// If we made it this far, we've failed.
+		const FString ErrorMsg = TEXT("Unable to locate Python interpreter. Run Switchboard first-time setup, or configure the Switchboard plugin PythonInterpreterPath setting.");
+		UE_LOG(LogSwitchboardPlugin, Error, TEXT("%s"), *ErrorMsg);
+		FPlatformMisc::MessageBoxExt(EAppMsgType::Ok, *ErrorMsg, TEXT("No Python interpreter found"));
+		return FString();
 	}
 
 	void OnLaunchSwitchboardClicked()
@@ -302,7 +284,7 @@ struct FSwitchboardMenuEntryImpl
 
 private:
 	TSharedPtr<FExtender> ToolBarExtender;
-	FString PathToListenerInstaller;
+	ECheckBoxState CachedAutolaunchCheckState;
 
 public:
 	static TUniquePtr<FSwitchboardMenuEntryImpl> Implementation;
diff --git a/Engine/Plugins/VirtualProduction/Switchboard/Source/SwitchboardEditor/Private/SwitchboardMenuEntry.h b/Engine/Plugins/VirtualProduction/Switchboard/Source/SwitchboardEditor/Private/SwitchboardMenuEntry.h
--- a/Engine/Plugins/VirtualProduction/Switchboard/Source/SwitchboardEditor/Private/SwitchboardMenuEntry.h
+++ b/Engine/Plugins/VirtualProduction/Switchboard/Source/SwitchboardEditor/Private/SwitchboardMenuEntry.h
@@ -4,8 +4,6 @@
 
 #include "CoreMinimal.h"
 
-DECLARE_LOG_CATEGORY_EXTERN(LogSwitchboardPlugin, Log, All);
-
 
 class FSwitchboardMenuEntry
 {
diff --git a/Engine/Plugins/VirtualProduction/Switchboard/Source/SwitchboardEditor/Public/SwitchboardEditorSettings.h b/Engine/Plugins/VirtualProduction/Switchboard/Source/SwitchboardEditor/Public/SwitchboardEditorSettings.h
--- a/Engine/Plugins/VirtualProduction/Switchboard/Source/SwitchboardEditor/Public/SwitchboardEditorSettings.h
+++ b/Engine/Plugins/VirtualProduction/Switchboard/Source/SwitchboardEditor/Public/SwitchboardEditorSettings.h
@@ -4,6 +4,7 @@
 
 #include "CoreMinimal.h"
 #include "Engine/EngineTypes.h"
+#include "Misc/Paths.h"
 #include "UObject/Object.h"
 #include "UObject/ObjectMacros.h"
 
@@ -18,8 +19,6 @@ class SWITCHBOARDEDITOR_API USwitchboardEditorSettings : public UObject
 	USwitchboardEditorSettings();
 
 public:
-
-
 	/** Switchboard installs its own python interpreter on its first launch. If you prefer
 	 * to use your own, specify the path to the python executable.
 	 */
@@ -50,4 +49,16 @@ public:
 	UFUNCTION(BlueprintPure, Category = "Switchboard")
 	static USwitchboardEditorSettings* GetSwitchboardEditorSettings();
 
+public:
+	FString GetListenerPlatformPath() const
+	{
+		FString ListenerPlatformPath = ListenerPath.FilePath;
+		FPaths::MakePlatformFilename(ListenerPlatformPath);
+		return ListenerPlatformPath;
+	}
+
+	FString GetListenerInvocation() const
+	{
+		return FString::Printf(TEXT("\"%s\" %s"), *GetListenerPlatformPath(), *ListenerCommandlineArguments);
+	}
 };
diff --git a/Engine/Plugins/VirtualProduction/Switchboard/Source/install_listener.py b/Engine/Plugins/VirtualProduction/Switchboard/Source/install_listener.py
deleted file mode 100644
