Switchboard: fix KeyError when the last Device of a particular category is deleted

diff --git a/Engine/Plugins/VirtualProduction/Switchboard/Source/Switchboard/switchboard/device_list_widget.py b/Engine/Plugins/VirtualProduction/Switchboard/Source/Switchboard/switchboard/device_list_widget.py
--- a/Engine/Plugins/VirtualProduction/Switchboard/Source/Switchboard/switchboard/device_list_widget.py
+++ b/Engine/Plugins/VirtualProduction/Switchboard/Source/Switchboard/switchboard/device_list_widget.py
@@ -152,13 +152,18 @@ class DeviceListWidget(QtWidgets.QListWidget):
         return self._device_widgets.values()
 
     def update_category_status(self, category_name, devices):
+        header_item = self._header_by_category_name.get(category_name)
+        if not header_item:
+            # The last device in the category may have just been deleted, in
+            # which case there'll be no header item to update.
+            return
+
         any_connected = False
         any_opened = False
         for device in devices:
             any_connected |= (not device.is_disconnected)
             any_opened |= (device.status > DeviceStatus.CLOSED)
 
-        header_item = self._header_by_category_name[category_name]
         header_widget = self.itemWidget(header_item)
         header_widget.update_connection_state(any_connected)
         header_widget.update_opened_state(any_connected, any_opened)
