Switchboard: add an nDisplay-specific setting for executable filename

This allows the executable for nDisplay devices to be specified independently
from the executable for Unreal devices.

diff --git a/Engine/Plugins/VirtualProduction/Switchboard/Source/Switchboard/switchboard/devices/ndisplay/plugin_ndisplay.py b/Engine/Plugins/VirtualProduction/Switchboard/Source/Switchboard/switchboard/devices/ndisplay/plugin_ndisplay.py
--- a/Engine/Plugins/VirtualProduction/Switchboard/Source/Switchboard/switchboard/devices/ndisplay/plugin_ndisplay.py
+++ b/Engine/Plugins/VirtualProduction/Switchboard/Source/Switchboard/switchboard/devices/ndisplay/plugin_ndisplay.py
@@ -260,6 +260,11 @@ class DevicenDisplay(DeviceUnreal):
             value="Mono", 
             possible_values=["Mono", "Frame sequential", "Side-by-Side", "Top-bottom"]
         ),
+        'executable_filename': Setting(
+            attr_name="executable_filename",
+            nice_name="nDisplay Executable Filename",
+            value="UE4Editor.exe",
+        ),
         'ndisplay_cmd_args': Setting(
             attr_name="ndisplay_cmd_args", 
             nice_name="Extra Cmd Line Args", 
@@ -417,6 +422,12 @@ class DevicenDisplay(DeviceUnreal):
     def is_recording_device(self):
         return False
 
+    # Override this property from the DeviceUnreal base class to use the
+    # nDisplay-specific setting.
+    @property
+    def executable_filename(self):
+        return DevicenDisplay.csettings["executable_filename"].get_value()
+
     def generate_unreal_command_line(self, map_name=""):
         uproject = os.path.normpath(CONFIG.UPROJECT_PATH.get_value(self.name))
 
diff --git a/Engine/Plugins/VirtualProduction/Switchboard/Source/Switchboard/switchboard/devices/unreal/plugin_unreal.py b/Engine/Plugins/VirtualProduction/Switchboard/Source/Switchboard/switchboard/devices/unreal/plugin_unreal.py
--- a/Engine/Plugins/VirtualProduction/Switchboard/Source/Switchboard/switchboard/devices/unreal/plugin_unreal.py
+++ b/Engine/Plugins/VirtualProduction/Switchboard/Source/Switchboard/switchboard/devices/unreal/plugin_unreal.py
@@ -779,8 +779,12 @@ class DeviceUnreal(Device):
         _, msg = message_protocol.create_fixExeFlags_message(puuid)
         self.unreal_client.send_message(msg)
 
+    @property
+    def executable_filename(self):
+        return DeviceUnreal.csettings["ue4_exe"].get_value()
+
     def generate_unreal_exe_path(self):
-        return CONFIG.engine_exe_path(CONFIG.ENGINE_DIR.get_value(self.name), DeviceUnreal.csettings["ue4_exe"].get_value())
+        return CONFIG.engine_exe_path(CONFIG.ENGINE_DIR.get_value(self.name), self.executable_filename)
 
     def get_vproles(self):
         ''' Gets selected vp roles that are also present in the ini file
