Switchboard: Added device context menu items to open last fetched log, utrace, and copy last launch command to clipboard.

diff --git a/Engine/Plugins/VirtualProduction/Switchboard/Source/Switchboard/switchboard/device_list_widget.py b/Engine/Plugins/VirtualProduction/Switchboard/Source/Switchboard/switchboard/device_list_widget.py
--- a/Engine/Plugins/VirtualProduction/Switchboard/Source/Switchboard/switchboard/device_list_widget.py
+++ b/Engine/Plugins/VirtualProduction/Switchboard/Source/Switchboard/switchboard/device_list_widget.py
@@ -42,11 +42,22 @@ class DeviceListWidget(QtWidgets.QListWidget):
         pass
 
     def contextMenuEvent(self, event):
-        if self.itemAt(event.pos()):
-            device_context_menu = QtWidgets.QMenu(self)
-            device_context_menu.addAction("Remove Device", lambda: self.ask_to_confirm_device_removal(event.pos()))
-            device_context_menu.exec_(event.globalPos())
-            event.accept()
+
+        item = self.itemAt(event.pos())
+
+        if not item:
+            return
+
+        item_widget = self.itemWidget(item)
+        assert item_widget
+
+        device_context_menu = QtWidgets.QMenu(self)
+
+        device_context_menu.addAction("Remove device", lambda: self.ask_to_confirm_device_removal(event.pos()))
+        item_widget.populate_context_menu(device_context_menu)
+
+        device_context_menu.exec_(event.globalPos())
+        event.accept()
 
     def ask_to_confirm_device_removal(self, pos):
         item = self.itemAt(pos)
diff --git a/Engine/Plugins/VirtualProduction/Switchboard/Source/Switchboard/switchboard/devices/device_widget_base.py b/Engine/Plugins/VirtualProduction/Switchboard/Source/Switchboard/switchboard/devices/device_widget_base.py
--- a/Engine/Plugins/VirtualProduction/Switchboard/Source/Switchboard/switchboard/devices/device_widget_base.py
+++ b/Engine/Plugins/VirtualProduction/Switchboard/Source/Switchboard/switchboard/devices/device_widget_base.py
@@ -304,6 +304,10 @@ class DeviceWidget(QtWidgets.QWidget):
         self.assign_button_to_name(name, button)
         return button
 
+    def populate_context_menu(self, cmenu: QtWidgets.QMenu):
+        ''' Called to populate the given context menu with any desired actions'''
+        pass
+
 
 class AddDeviceDialog(QtWidgets.QDialog):
     def __init__(self, device_type, existing_devices, parent=None):
diff --git a/Engine/Plugins/VirtualProduction/Switchboard/Source/Switchboard/switchboard/devices/unreal/plugin_unreal.py b/Engine/Plugins/VirtualProduction/Switchboard/Source/Switchboard/switchboard/devices/unreal/plugin_unreal.py
--- a/Engine/Plugins/VirtualProduction/Switchboard/Source/Switchboard/switchboard/devices/unreal/plugin_unreal.py
+++ b/Engine/Plugins/VirtualProduction/Switchboard/Source/Switchboard/switchboard/devices/unreal/plugin_unreal.py
@@ -22,8 +22,8 @@ from switchboard import message_protocol
 from switchboard import switchboard_application
 from switchboard import switchboard_utils as sb_utils
 from switchboard.config import CONFIG, BoolSetting, DirectoryPathSetting, \
-    IntSetting, MultiOptionSetting, OptionSetting, StringSetting, SETTINGS, \
-    DEFAULT_MAP_TEXT
+    IntSetting, MultiOptionSetting, OptionSetting, StringSetting, FileSystemPathSetting, \
+    SETTINGS, DEFAULT_MAP_TEXT
 from switchboard.devices.device_base import Device, DeviceStatus, \
     PluginHeaderWidgets
 from switchboard.devices.device_widget_base import DeviceWidget, DeviceAutoJoinMUServerUI
@@ -31,6 +31,8 @@ from switchboard.listener_client import ListenerClient
 from switchboard.switchboard_logging import LOGGER
 import switchboard.switchboard_widgets as sb_widgets
 
+from switchboard.tools.insights_launcher import InsightsLauncher
+
 from .listener_watcher import ListenerWatcher
 from .redeploy_dialog import RedeployListenerDialog
 from . import version_helpers
@@ -475,6 +477,27 @@ class DeviceUnreal(Device):
             show_ui=False
         )
 
+        self.last_launch_command = StringSetting(
+            attr_name="last_launch_command",
+            nice_name="Last Launch Command",
+            value=kwargs.get("last_launch_command",''),
+            show_ui=False
+        )
+
+        self.last_log_path = FileSystemPathSetting(
+            attr_name="last_log_path",
+            nice_name="Last Log Path",
+            value=kwargs.get("last_log_path",''),
+            show_ui=False
+        )
+
+        self.last_trace_path = FileSystemPathSetting(
+            attr_name="last_trace_path",
+            nice_name="Last Insights Trace Path",
+            value=kwargs.get("last_trace_path",''),
+            show_ui=False
+        )
+
         self.setting_ip_address.signal_setting_changed.connect(
             self.on_setting_ip_address_changed)
         DeviceUnreal.csettings['port'].signal_setting_changed.connect(
@@ -673,7 +696,10 @@ class DeviceUnreal(Device):
     def device_settings(self):
         return super().device_settings() + [
             self.setting_roles,
-            self.autojoin_mu_server
+            self.autojoin_mu_server,
+            self.last_launch_command,
+            self.last_log_path,
+            self.last_trace_path,
         ]
 
     def check_settings_valid(self) -> bool:
@@ -724,6 +750,15 @@ class DeviceUnreal(Device):
 
         super().device_widget_registered(device_widget)
 
+        # hook to open last log signal from widget
+        device_widget.signal_open_last_log.connect(self.on_open_last_log)
+
+        # hook to open last trace signal from widget
+        device_widget.signal_open_last_trace.connect(self.on_open_last_trace)
+
+        # hook to copy last launch signal from widget
+        device_widget.signal_copy_last_launch_command.connect(self.on_copy_last_launch_command)
+
         self.update_settings_menu_state()
 
         device_widget.autojoin_mu.signal_device_widget_autojoin_mu.connect(self.on_autojoin_mu_ui_change)
@@ -1233,7 +1268,7 @@ class DeviceUnreal(Device):
                               f'-CONCERTDISPLAYNAME="{self.name}"')
 
         if CONFIG.INSIGHTS_TRACE_ENABLE.get_value():
-            LOGGER.warning(f"Unreal Insight Tracing is enabled for '{self.name}'. This may effect Unreal Engine performance.")
+            LOGGER.warning(f"Unreal Insight Tracing is enabled for '{self.name}'. This may affect Unreal Engine performance.")
             remote_utrace_path = self.get_utrace_filepath()
             command_line_args += ' -statnamedevents' if CONFIG.INSIGHTS_STAT_EVENTS.get_value() else ''
             command_line_args += ' -tracefile="{}" -trace="{}"'.format(
@@ -1345,6 +1380,8 @@ class DeviceUnreal(Device):
         # TODO: Sanitize these on Qt input? Deserialization?
         args = args.replace('\r', ' ').replace('\n', ' ')
 
+        self.last_launch_command.update_value(f'{engine_path} {args}')
+
         puuid, msg = message_protocol.create_start_process_message(
             prog_path=engine_path,
             prog_args=args,
@@ -1785,11 +1822,12 @@ class DeviceUnreal(Device):
 
     def backup_file(self, filename):
         ''' Rotate existing filename to a timestamped backup, a la Unreal. '''
-        log_download_path = self.get_log_download_dir()
-        if not log_download_path:
+
+        try:
+            src = self.make_local_filepath_to_fetch(filename)
+        except:
             return
 
-        src = log_download_path / filename
         if not src.is_file():
             return
 
@@ -1797,7 +1835,7 @@ class DeviceUnreal(Device):
         modtime_str = modtime.strftime('%Y.%m.%d-%H.%M.%S')
         dest_filename = f'{src.stem}-backup-{modtime_str}{src.suffix}'
 
-        src.rename(log_download_path / dest_filename)
+        src.rename(self.make_local_filepath_to_fetch(dest_filename))
 
     def get_remote_log_path(self):
         remote_project_path = \
@@ -1843,6 +1881,18 @@ class DeviceUnreal(Device):
         # TODO: sync crash log?
         return True
 
+    def make_local_filepath_to_fetch(self, filename):
+        """
+        Makes the destination path of a file to be fetched
+        """
+
+        log_download_dir = self.get_log_download_dir()
+
+        if not log_download_dir:
+            raise NotADirectoryError
+
+        return log_download_dir / filename
+
     def start_retrieve_utrace(self, unreal_exit_code:int ):
         """
         Retrieve the utrace file if tracing is enabled.
@@ -1852,7 +1902,15 @@ class DeviceUnreal(Device):
 
         remote_path = self.get_utrace_filepath()
 
-        self.backup_file(os.path.basename(remote_path))
+        filename = os.path.basename(remote_path)
+        self.backup_file(filename)
+
+        # remember the path to the trace to be fetched, so that it can later be opened from the device's context menu
+        try:
+            self.last_trace_path.update_value(str(self.make_local_filepath_to_fetch(filename)))
+        except:
+            pass
+
         return self.fetch_file(remote_path)
 
     def start_retrieve_log(self, unreal_exit_code: int):
@@ -1861,6 +1919,12 @@ class DeviceUnreal(Device):
         """
         self.backup_file(self.log_filename)
 
+        # remember the path to the log file to be fetched, so that it can later be opened from the device's context menu
+        try:
+            self.last_log_path.update_value(str(self.make_local_filepath_to_fetch(self.log_filename)))
+        except:
+            pass
+
         remote_log_path = self.get_remote_log_path() / self.log_filename
         return self.fetch_file(remote_log_path)
 
@@ -1911,6 +1975,26 @@ class DeviceUnreal(Device):
         self.autojoin_mu_server.update_value(autojoin.is_autojoin_enabled())
         self.update_settings_menu_state()
 
+    def on_open_last_log(self):
+        ''' Opens the last log in your preferred editor '''
+        sb_utils.openfile_with_default_app(self.last_log_path.get_value())
+
+    def on_open_last_trace(self):
+        ''' Opens the last unreal insights trace '''
+
+        tracepath = self.last_trace_path.get_value()
+
+        if not pathlib.Path(tracepath).is_file():
+            LOGGER.error(f"Could not find '{tracepath}'")
+            return
+
+        insights = InsightsLauncher()
+        insights.launch(args=[str(tracepath)], allow_duplicate=True)
+
+    def on_copy_last_launch_command(self):
+        ''' Copies the last launch command to the clipboard'''
+        sb_utils.copy2clipboard(self.last_launch_command.get_value())
+
 def parse_unreal_tag_file(file_content):
     tags = []
     for line in file_content:
@@ -1922,6 +2006,11 @@ def parse_unreal_tag_file(file_content):
     return tags
 
 class DeviceWidgetUnreal(DeviceWidget):
+
+    signal_open_last_log = QtCore.Signal(object)
+    signal_open_last_trace = QtCore.Signal(object)
+    signal_copy_last_launch_command = QtCore.Signal(object)
+
     def __init__(self, name, device_hash, ip_address, icons, parent=None):
         self._autojoin_visible = True
 
@@ -2166,3 +2255,11 @@ class DeviceWidgetUnreal(DeviceWidget):
             self._connect()
         else:
             self._disconnect()
+
+    def populate_context_menu(self, cmenu: QtWidgets.QMenu):
+        ''' Called to populate the given context menu with any desired actions'''
+        
+        cmenu.addAction("Open fetched log", lambda: self.signal_open_last_log.emit(self))
+        cmenu.addAction("Open fetched trace", lambda: self.signal_open_last_trace.emit(self))
+        cmenu.addAction("Copy last launch command", lambda: self.signal_copy_last_launch_command.emit(self))
+
diff --git a/Engine/Plugins/VirtualProduction/Switchboard/Source/Switchboard/switchboard/switchboard_utils.py b/Engine/Plugins/VirtualProduction/Switchboard/Source/Switchboard/switchboard/switchboard_utils.py
--- a/Engine/Plugins/VirtualProduction/Switchboard/Source/Switchboard/switchboard/switchboard_utils.py
+++ b/Engine/Plugins/VirtualProduction/Switchboard/Source/Switchboard/switchboard/switchboard_utils.py
@@ -1,12 +1,15 @@
 # Copyright Epic Games, Inc. All Rights Reserved.
+
 from .switchboard_logging import LOGGER
 
 from enum import Enum
+
 import datetime
 import subprocess
 import os
 import threading
 import sys
+import pathlib
 
 
 class PriorityModifier(Enum):
@@ -181,3 +184,24 @@ def explore_path(path:str):
         subprocess.Popen(f'open {path}')
     else:
         LOGGER.error(f"explore_path not supported in platform '{sys.platform}'")
+
+def copy2clipboard(text:str):
+    if sys.platform.startswith('win'):
+        subprocess.Popen(f'echo[{text.strip()}|clip', shell=True)
+    else:
+        LOGGER.error(f"copy2clipboard not supported in platform '{sys.platform}'")
+
+def openfile_with_default_app(path:str):
+
+    if not pathlib.Path(path).is_file():
+        LOGGER.error(f"openfile_with_default_app could not find '{path}'")
+        return
+
+    if sys.platform.startswith('win'):
+        os.startfile(path)
+    elif sys.platform.startswith('linux'):
+        subprocess.Popen(f'xdg-open {path}')
+    elif sys.platform.startswith('darwin'):
+        subprocess.Popen(f'open {path}')
+    else:
+        LOGGER.error(f"openfile_with_default_app not supported in platform '{sys.platform}'")
diff --git a/Engine/Plugins/VirtualProduction/Switchboard/Source/Switchboard/switchboard/tools/application_launcher.py b/Engine/Plugins/VirtualProduction/Switchboard/Source/Switchboard/switchboard/tools/application_launcher.py
--- a/Engine/Plugins/VirtualProduction/Switchboard/Source/Switchboard/switchboard/tools/application_launcher.py
+++ b/Engine/Plugins/VirtualProduction/Switchboard/Source/Switchboard/switchboard/tools/application_launcher.py
@@ -24,7 +24,7 @@ class ApplicationLauncher(object):
     def exe_name(self):
         return os.path.split(self.exe_path())[1]
 
-    def launch(self, args:Optional[List[str]] = []):
+    def launch(self, args:Optional[List[str]] = [], allow_duplicate:bool = False):
         ''' Launches this application with the given arguments '''
 
         args.insert(0, self.exe_path())
@@ -32,7 +32,7 @@ class ApplicationLauncher(object):
 
         with self.lock:
 
-            if self.is_running():
+            if (not allow_duplicate) and self.is_running():
                 return False
 
             if not os.path.exists(self.exe_path()):
