Switchboard: Removed soft kill button.

This had been originally added to mitigate a gpu hang, which happened when the Unreal process was terminated while waiting on Nvidia quadro sync custom present, by synchronizing the exit across the cluster via a cluster event outside present. But Nvidia has fixed that bug in R470, so this should not be necessary anymore.

diff --git a/Engine/Plugins/VirtualProduction/Switchboard/Source/Switchboard/switchboard/devices/ndisplay/ndisplay_monitor_ui.py b/Engine/Plugins/VirtualProduction/Switchboard/Source/Switchboard/switchboard/devices/ndisplay/ndisplay_monitor_ui.py
--- a/Engine/Plugins/VirtualProduction/Switchboard/Source/Switchboard/switchboard/devices/ndisplay/ndisplay_monitor_ui.py
+++ b/Engine/Plugins/VirtualProduction/Switchboard/Source/Switchboard/switchboard/devices/ndisplay/ndisplay_monitor_ui.py
@@ -74,13 +74,6 @@ class nDisplayMonitorUI(QWidget):
         btn_fix_exe_flags.clicked.connect(
             self.monitor.on_fix_exe_flags_clicked)
 
-        btn_soft_kill = QPushButton('Soft Kill')
-        btn_soft_kill.setToolTip(
-            'Sends a message to the master node to terminate the session.\n'
-            'This is preferable to the normal kill button because it ensures '
-            'the nodes exit properly')
-        btn_soft_kill.clicked.connect(self.monitor.on_soft_kill_clicked)
-
         btn_minimize_windows = QPushButton('Minimize')
         btn_minimize_windows.setToolTip(
             'Minimizes all windows in the nodes.')
@@ -94,7 +87,6 @@ class nDisplayMonitorUI(QWidget):
         layout_buttons.addWidget(btn_refresh_mosaics)
         layout_buttons.addWidget(btn_fix_exe_flags)
         layout_buttons.addWidget(btn_minimize_windows)
-        layout_buttons.addWidget(btn_soft_kill)
 
         return layout_buttons
 
