Switchboard: Added 'bookmark' to default trace channels. In preparation to vblank bookmarks coming soon.

diff --git a/Engine/Plugins/VirtualProduction/Switchboard/Source/Switchboard/switchboard/config.py b/Engine/Plugins/VirtualProduction/Switchboard/Source/Switchboard/switchboard/config.py
--- a/Engine/Plugins/VirtualProduction/Switchboard/Source/Switchboard/switchboard/config.py
+++ b/Engine/Plugins/VirtualProduction/Switchboard/Source/Switchboard/switchboard/config.py
@@ -1405,7 +1405,7 @@ class Config(object):
         self.INSIGHTS_TRACE_ARGS = StringSetting(
             "tracing_args",
             "Unreal Insights Tracing Args",
-            data.get('tracing_arg', 'log,cpu,gpu,frame,concert,messaging')
+            data.get('tracing_arg', 'log,cpu,gpu,frame,bookmark,concert,messaging')
         )
         self.INSIGHTS_STAT_EVENTS = BoolSetting(
             "tracing_stat_events",
diff --git a/Engine/Plugins/VirtualProduction/Switchboard/Source/Switchboard/switchboard/switchboard_dialog.py b/Engine/Plugins/VirtualProduction/Switchboard/Source/Switchboard/switchboard/switchboard_dialog.py
--- a/Engine/Plugins/VirtualProduction/Switchboard/Source/Switchboard/switchboard/switchboard_dialog.py
+++ b/Engine/Plugins/VirtualProduction/Switchboard/Source/Switchboard/switchboard/switchboard_dialog.py
@@ -75,7 +75,7 @@ class DeviceAdditionalSettingsUI(QtCore.QObject):
     signal_device_widget_tracing = QtCore.Signal(object)
 
     enable_insight_trace = False
-    insight_trace_args = "log,cpu,gpu,frame,concert,messaging"
+    insight_trace_args = "log,cpu,gpu,frame,bookmark,concert,messaging"
 
     def __init__(self, name, parent = None):
         super().__init__(parent)
