Force the mu server to terminate by polling the process tree versus using the process handle.

diff --git a/Engine/Plugins/VirtualProduction/Switchboard/Source/Switchboard/switchboard/devices/unreal/plugin_unreal.py b/Engine/Plugins/VirtualProduction/Switchboard/Source/Switchboard/switchboard/devices/unreal/plugin_unreal.py
--- a/Engine/Plugins/VirtualProduction/Switchboard/Source/Switchboard/switchboard/devices/unreal/plugin_unreal.py
+++ b/Engine/Plugins/VirtualProduction/Switchboard/Source/Switchboard/switchboard/devices/unreal/plugin_unreal.py
@@ -1069,7 +1069,7 @@ class DeviceUnreal(Device):
                         QtWidgets.QMessageBox.Yes | QtWidgets.QMessageBox.No)
 
                     if mb_ret == QtWidgets.QMessageBox.Yes:
-                        DeviceUnreal.mu_server.terminate()
+                        DeviceUnreal.mu_server.terminate(bypolling=True)
 
                 puuid_dependency = self._build_mu_server(
                     puuid_dependency=puuid_dependency)
