Add search bar to project settings

diff --git a/Engine/Plugins/VirtualProduction/Switchboard/Source/Switchboard/switchboard/settings_dialog.py b/Engine/Plugins/VirtualProduction/Switchboard/Source/Switchboard/switchboard/settings_dialog.py
--- a/Engine/Plugins/VirtualProduction/Switchboard/Source/Switchboard/switchboard/settings_dialog.py
+++ b/Engine/Plugins/VirtualProduction/Switchboard/Source/Switchboard/switchboard/settings_dialog.py
@@ -11,6 +11,7 @@ from PySide2 import QtWidgets
 from switchboard import config
 from switchboard import switchboard_widgets as sb_widgets
 from switchboard.config import CONFIG, SETTINGS
+from switchboard.settings_search import SettingsSearch
 
 RELATIVE_PATH = os.path.dirname(__file__)
 
@@ -23,6 +24,8 @@ class SettingsDialog(QtCore.QObject):
         loader = QtUiTools.QUiLoader()
         self.ui = loader.load(os.path.join(RELATIVE_PATH, "ui/settings.ui"))
 
+        self.settings_search = SettingsSearch([self.ui.scrollAreaWidgetContents], self.ui.scrollArea)
+
         self.ui.setWindowTitle("Settings")
 
         self.ui.finished.connect(self._on_finished)
@@ -38,6 +41,9 @@ class SettingsDialog(QtCore.QObject):
         self._config_paths = config.list_config_paths()
         self.set_config_path(SETTINGS.CONFIG)
 
+        self.ui.searchBar.textChanged.connect(
+            self._on_search_text_edited)
+
         self.ui.config_path_line_edit.textChanged.connect(
             self.config_path_text_changed)
         self.ui.uproject_browse_button.clicked.connect(
@@ -87,9 +93,8 @@ class SettingsDialog(QtCore.QObject):
         # the close button as opposed to ok/cancel buttons, so we intercept the
         # close to issue a warning if the config path was changed and we're
         # about to overwrite some other existing config file.
-        if (self._changed_config_path and
-                self._changed_config_path in self._config_paths and
-                self._changed_config_path != self._current_config_path):
+        if (self._changed_config_path in self._config_paths 
+                and self._changed_config_path != self._current_config_path):
             # Show the confirmation dialog using a relative path to the config.
             rel_config_path = config.get_relative_config_path(
                 self._changed_config_path)
@@ -104,6 +109,9 @@ class SettingsDialog(QtCore.QObject):
                 # originally specified path when queried via config_path().
                 self._changed_config_path = None
 
+    def _on_search_text_edited(self, search_string: str):
+        self.settings_search.search(search_string)
+
     def config_path(self):
         if self._changed_config_path:
             return self._changed_config_path
@@ -295,7 +303,7 @@ class SettingsDialog(QtCore.QObject):
 
     # Devices
     def add_section_for_plugin(
-            self, plugin_name, plugin_settings, device_settings):
+        self, plugin_name, plugin_settings, device_settings):
         any_device_settings = (
             any([device[1] for device in device_settings]) or
             any([device[2] for device in device_settings]))
@@ -308,6 +316,8 @@ class SettingsDialog(QtCore.QObject):
         if device_override_group_box.parent() is None:
             device_override_group_box.setTitle(f'{plugin_name} Settings')
             device_override_group_box.setLayout(QtWidgets.QVBoxLayout())
+            device_override_group_box.setSizePolicy(
+                QtWidgets.QSizePolicy(QtWidgets.QSizePolicy.MinimumExpanding, QtWidgets.QSizePolicy.Maximum))
             self.ui.device_override_layout.addWidget(device_override_group_box)
 
         plugin_layout = QtWidgets.QFormLayout()
diff --git a/Engine/Plugins/VirtualProduction/Switchboard/Source/Switchboard/switchboard/settings_search.py b/Engine/Plugins/VirtualProduction/Switchboard/Source/Switchboard/switchboard/settings_search.py
new file mode 100644
--- /dev/null
+++ b/Engine/Plugins/VirtualProduction/Switchboard/Source/Switchboard/switchboard/settings_search.py
@@ -0,0 +1,223 @@
+# Copyright Epic Games, Inc. All Rights Reserved.
+
+import os
+import pathlib
+from typing import List, Dict
+
+from PySide2 import QtCore
+from PySide2 import QtGui
+from PySide2 import QtUiTools
+from PySide2 import QtWidgets
+
+
+def as_widget_or_layout(row_item):
+    return row_item.widget() if row_item.widget() else row_item.layout()
+
+
+class FormLayoutRowVisibility:
+    """
+    Handles the showing and hiding of rows in QFormLayout
+    
+    Row widgets are removed when hiding and added back when showing instead of calling setVisible on the row's widgets:
+    Even if a row only contains invisible widgets, QFormLayout.layoutVerticalSpacing still affects them and visually 
+    shifts visible rows around.
+    """
+    
+    def __init__(self, form: QtWidgets.QFormLayout):
+        self.form = form
+        # Save the original widget order so we can preserve it when adding hidden widgets back
+        self.rows = []
+        for i in range(form.rowCount()):
+            label = form.itemAt(i, QtWidgets.QFormLayout.LabelRole)
+            field = form.itemAt(i, QtWidgets.QFormLayout.FieldRole)
+            if label and field:
+                pair = (
+                    as_widget_or_layout(label),
+                    as_widget_or_layout(field),
+                    True
+                )
+                self.rows.append(pair)
+        
+    def row_count(self):
+        return len(self.rows)
+        
+    def get_label_item(self, index: int):
+        return self.rows[index][0]
+        
+    def set_row_visible(self, index: int, visible: bool):
+        (label, field, is_visible) = self.rows[index]
+        if visible and not is_visible:
+            self._show_row(index)
+        elif not visible and is_visible:
+            self._hide_row(index)
+                    
+    def _show_row(self, index: int):
+        (label, field, is_visible) = self.rows[index]
+        index_to_use = 0 if index == 0 else self._find_index_to_insert_row(index)
+        self.form.insertRow(index_to_use, label, field)
+        self.rows[index] = (
+            label,
+            field,
+            True
+        )
+        
+        # Widgets were hidden when the row was hidden
+        self._set_visible_recursive(label, True)
+        self._set_visible_recursive(field, True)
+        
+    def _hide_row(self, index: int):
+        (label, field, is_visible) = self.rows[index]
+        # "index" is the original index of the row before we any rows were hidden.
+        # Since we may have hidden certain rows already, we must search for the the row that corresponds "index"
+        for i in range(self.form.rowCount()):
+            label_item = self.form.itemAt(i, QtWidgets.QFormLayout.LabelRole)
+            if label == as_widget_or_layout(label_item):
+                self._take_row(i)
+                self.rows[index] = (label, field, False)
+                # Found the matching row
+                break
+
+    def _take_row(self, index: int):
+        label_item = self.form.itemAt(index, QtWidgets.QFormLayout.LabelRole)
+        field_item = self.form.itemAt(index, QtWidgets.QFormLayout.FieldRole)
+        
+        # Otherwise there will be artifacts in the UI texture
+        self._set_visible_recursive(as_widget_or_layout(label_item), False)
+        self._set_visible_recursive(as_widget_or_layout(field_item), False)
+        
+        # Remove items before removeRow so the widgets are not destroyed
+        self.form.removeItem(label_item)
+        self.form.removeItem(field_item)
+        self.form.removeRow(index)
+        
+        # Reallocate space
+        self.form.invalidate()
+        self.form.update()
+
+    def _find_index_to_insert_row(self, index: int):
+        (parent_label, _, _) = self.rows[index - 1]
+        for actual_form_index in range(self.form.rowCount()):
+            actual_item = self.form.itemAt(actual_form_index, QtWidgets.QFormLayout.LabelRole)
+            for saved_form_index in range(index):
+                if as_widget_or_layout(actual_item) == parent_label:
+                    return actual_form_index + 1
+
+        return 0
+            
+    def _set_visible_recursive(self, element, visible: bool):
+        if isinstance(element, QtWidgets.QWidget):
+            element.setVisible(visible)
+            return
+        
+        for i in range(element.count()):
+            child_item = element.itemAt(i)
+            if child_item.widget():
+                child_item.widget().setVisible(visible)
+            elif child_item.layout():
+                self._set_visible_recursive(child_item.layout(), visible)
+
+
+class SettingsSearch:
+    def __init__(self, searched_widgets: List[QtWidgets.QWidget], scroll_area: QtWidgets.QScrollArea):
+        self.reset_search_bar_height = 0
+        self.last_search_string_len = 0
+        
+        self.searched_widgets = searched_widgets
+        self.scroll_area = scroll_area
+        
+        form_layout_handlers: Dict[QtWidgets.QFormLayout, FormLayoutRowVisibility] = {}
+        self.form_layout_handlers = form_layout_handlers
+        
+    def search(self, search_string: str):
+        # The first time something is typed, save scroll height to restore later
+        self._take_scroll_snapshot(search_string)
+        
+        search_term_list = search_string.split()
+        for widget in self.searched_widgets:
+            self._update_widget_visibility(widget, search_term_list)
+
+        # Restore scroll bar height when search is over
+        self._restore_scroll_snapshot(search_string)
+
+    def _take_scroll_snapshot(self, search_string: str):
+        just_started_search = self.last_search_string_len == 0
+        if just_started_search:
+            self.reset_search_bar_height = self.scroll_area.verticalScrollBar().value()
+            
+            self.scroll_area.verticalScrollBar().setValue(0)
+
+    def _restore_scroll_snapshot(self, search_string: str):
+        len_search = len(search_string)
+        is_search_empty = len_search < self.last_search_string_len and len_search == 0
+        if is_search_empty:
+            # Need to wait for UI to update layout...
+            QtCore.QTimer.singleShot(
+                25,
+                lambda height=self.reset_search_bar_height: self.scroll_area.verticalScrollBar().setValue(height)
+            )
+        self.last_search_string_len = len_search
+
+    def _update_widget_visibility(self, widget: QtWidgets.QWidget, search_term_list):
+        # We only search labels for search terms
+        if isinstance(widget, QtWidgets.QLabel):
+            matches_search = self._is_search_match(widget.text(), search_term_list)
+            widget.setVisible(matches_search)
+            return matches_search
+        
+        if widget.layout():
+            is_category = isinstance(widget, QtWidgets.QGroupBox)
+            is_match_on_containing_category = is_category and self._is_search_match(widget.title(), search_term_list)
+            # Make all children appear
+            if is_match_on_containing_category:
+                search_term_list = []
+
+            is_any_child_visible = self._update_layout_visibility(widget.layout(), search_term_list) \
+                or is_match_on_containing_category
+            widget.setVisible(is_any_child_visible)
+            return is_any_child_visible
+
+        return False
+
+    def _update_layout_visibility(self, layout: QtWidgets.QLayout, search_term_list):
+        is_match_on_any_child = False
+
+        if isinstance(layout, QtWidgets.QFormLayout):
+            return self._handle_form_layout(layout, search_term_list)
+
+        # Is this horizontal boxes similar to a form?
+        if isinstance(layout, QtWidgets.QHBoxLayout) and layout.count() == 2:
+            first = layout.itemAt(0)
+            first_item_is_label = first and first.widget() and isinstance(first.widget(), QtWidgets.QLabel)
+            if first_item_is_label:
+                is_match_on_any_child = self._update_widget_visibility(
+                    first.widget(), search_term_list)
+                return is_match_on_any_child
+
+        for i in range(layout.count()):
+            layout_item = layout.itemAt(i)
+            if layout_item.widget():
+                is_match_on_any_child |= self._update_widget_visibility(layout_item.widget(), search_term_list)
+            elif layout_item.layout():
+                is_match_on_any_child |= self._update_layout_visibility(layout_item.layout(), search_term_list)
+
+        return is_match_on_any_child
+    
+    def _handle_form_layout(self, layout: QtWidgets.QFormLayout, search_term_list):
+        is_match_on_any_child = False
+        form_handler = self.form_layout_handlers.setdefault(layout, FormLayoutRowVisibility(layout))
+        for i in range(form_handler.row_count()):
+            label = form_handler.get_label_item(i)
+            if isinstance(label, QtWidgets.QWidget):
+                is_match = self._update_widget_visibility(label, search_term_list)
+                is_match_on_any_child |= is_match
+                form_handler.set_row_visible(i, is_match)
+
+        return is_match_on_any_child
+
+    @staticmethod
+    def _is_search_match(text: str, search_term_list: str):
+        for search_term in search_term_list:
+            if not search_term.lower() in text.lower():
+                return False
+        return True
+    
\ No newline at end of file
