Switchboard: Added a script to update resources.py when a new image/icon is added.

Switchboard will read the resources from ../resources.py and not from resources.qrc.
So when adding a new image, copy it to /images and then manually reference it in resources.qrc.

Then run this script, which will check out ../resources.py from perforce and update it.

It requires pyside2-rcc.exe to be in the path (typically found inside <python>/Scripts).
You should also add resources.qrc and the image files to your perforce changelist.

diff --git a/Engine/Plugins/VirtualProduction/Switchboard/Source/Switchboard/switchboard/ui/update_resources_py.bat b/Engine/Plugins/VirtualProduction/Switchboard/Source/Switchboard/switchboard/ui/update_resources_py.bat
new file mode 100755
--- /dev/null
+++ b/Engine/Plugins/VirtualProduction/Switchboard/Source/Switchboard/switchboard/ui/update_resources_py.bat
@@ -0,0 +1,9 @@
+
+:: Switchboard will read the resources from ../resources.py and not from resources.qrc.
+:: So when adding a new image, copy it to /images and then manually reference it in resources.qrc.
+:: Then run this script, which will check out ../resources.py from perforce and update it.
+:: It requires pyside2-rcc.exe to be in the path (typically found inside <python>/Scripts).
+:: You should also add resources.qrc and the image files to your perforce changelist.
+
+p4 edit ../resources.py
+pyside2-rcc.exe resources.qrc -o ../resources.py
\ No newline at end of file
