Switchboard: Enable vblank monitoring via early cvar when traces and bookmarks are being captured.

diff --git a/Engine/Plugins/VirtualProduction/Switchboard/Source/Switchboard/switchboard/devices/ndisplay/plugin_ndisplay.py b/Engine/Plugins/VirtualProduction/Switchboard/Source/Switchboard/switchboard/devices/ndisplay/plugin_ndisplay.py
--- a/Engine/Plugins/VirtualProduction/Switchboard/Source/Switchboard/switchboard/devices/ndisplay/plugin_ndisplay.py
+++ b/Engine/Plugins/VirtualProduction/Switchboard/Source/Switchboard/switchboard/devices/ndisplay/plugin_ndisplay.py
@@ -672,11 +672,6 @@ class DevicenDisplay(DeviceUnreal):
     def generate_unreal_command_line(self, map_name=""):
         uproject = os.path.normpath(CONFIG.UPROJECT_PATH.get_value(self.name))
 
-        # Device profile CVars
-        dp_cvars = str(
-            self.csettings['ndisplay_dp_cvars'].get_value(self.name)).strip()
-        dp_cvars = f'-DPCVars="{dp_cvars}"' if len(dp_cvars) else ''
-
         # Extra arguments specified in settings
         additional_args = self.extra_cmdline_args_setting
 
@@ -806,7 +801,6 @@ class DevicenDisplay(DeviceUnreal):
             f'{friendly_name}',           # Stage Friendly Name
             f'{session_id}',              # Session ID.
             f'{max_gpu_count}',           # Max GPU count (mGPU)
-            f'{dp_cvars}',                # Device profile CVars
             f'-dc_cfg="{cfg_file}"',      # nDisplay config file
             f'{render_api}',              # dx11/12
             f'{render_mode}',             # mono/...
@@ -866,6 +860,10 @@ class DevicenDisplay(DeviceUnreal):
             '-CONCERTISHEADLESS',
         ])
 
+        # Device profile CVars.
+        dp_cvars = str(self.csettings['ndisplay_dp_cvars'].get_value(self.name)).split(',')
+        dp_cvars = [cvar.strip() for cvar in dp_cvars if len(cvar.strip()) and len(cvar.split('=')) == 2]
+
         # Insights traces parameters
         if CONFIG.INSIGHTS_TRACE_ENABLE.get_value():
 
@@ -873,19 +871,31 @@ class DevicenDisplay(DeviceUnreal):
 
             remote_utrace_path = self.get_utrace_filepath()
 
+            traces = CONFIG.INSIGHTS_TRACE_ARGS.get_value()
+
             args.extend([
                 f'-tracefile="{remote_utrace_path}"',
-                f'-trace="{CONFIG.INSIGHTS_TRACE_ARGS.get_value()}"'
+                f'-trace="{traces}"'
             ])
 
             if CONFIG.INSIGHTS_STAT_EVENTS.get_value():
                 args.append("-statnamedevents")
 
+            # if bookmarks are enabled, also enable the vblank monitoring thread
+            if 'bookmark' in traces.split(','):
+                dp_cvars.append('nDisplay.sync.diag.VBlankMonitoring=1')
+
+        # add accumulated dpcvars to args
+        if len(dp_cvars):
+            args.append(f"-DPCVars=\"{','.join(dp_cvars)}\"")
+
+        # Logging
         args.append(self.csettings['logging'].get_command_line_arg(
             override_device_name=self.name))
 
         path_to_exe = self.generate_unreal_exe_path()
         args_expanded = ' '.join(args)
+        
         self.settings['ue_command_line'].update_value(
             f"{path_to_exe} {args_expanded}")
 
