Switchboard: Device settings button relabeled to Insights capture settings.

Now it will only appear highlighted when the user has traces enabled.

The previous behavior was to be highlighted when any setting was different than the default.

diff --git a/Engine/Plugins/VirtualProduction/Switchboard/Source/Switchboard/switchboard/switchboard_dialog.py b/Engine/Plugins/VirtualProduction/Switchboard/Source/Switchboard/switchboard/switchboard_dialog.py
--- a/Engine/Plugins/VirtualProduction/Switchboard/Source/Switchboard/switchboard/switchboard_dialog.py
+++ b/Engine/Plugins/VirtualProduction/Switchboard/Source/Switchboard/switchboard/switchboard_dialog.py
@@ -83,12 +83,9 @@ class DeviceAdditionalSettingsUI(QtCore.QObject):
 
     def _update_enable_state(self):
         """
-        The button assigned for the settings UI should be "checked" when the user has modified a default.
+        The button assigned for the settings UI should be "checked" when the user has traces enabled.
         """
-        tracing_different = self.enable_insights_menu().isChecked() is not DeviceAdditionalSettingsUI.enable_insight_trace
-        tracing_values_different = self.insight_trace_args != DeviceAdditionalSettingsUI.insight_trace_args
-
-        self._button.setChecked(tracing_different or tracing_values_different)
+        self._button.setChecked(self.enable_insights_menu().isChecked())
 
     def set_insight_trace_state(self, is_checked):
         """
