Add tabs to settings

diff --git a/Engine/Plugins/VirtualProduction/Switchboard/Source/Switchboard/switchboard/settings_dialog.py b/Engine/Plugins/VirtualProduction/Switchboard/Source/Switchboard/switchboard/settings_dialog.py
--- a/Engine/Plugins/VirtualProduction/Switchboard/Source/Switchboard/switchboard/settings_dialog.py
+++ b/Engine/Plugins/VirtualProduction/Switchboard/Source/Switchboard/switchboard/settings_dialog.py
@@ -12,105 +12,423 @@
 from switchboard import switchboard_widgets as sb_widgets
 from switchboard.config import CONFIG, SETTINGS
 from switchboard.settings_search import SettingsSearch
+from switchboard.ui.horizontal_tabs import HorizontalTabWidget
 
 RELATIVE_PATH = os.path.dirname(__file__)
 
 
+def clear_widgets(layout):
+    for i in range(layout.count()):
+        layout.takeAt(0) 
+
+
 class SettingsDialog(QtCore.QObject):
     def __init__(self):
         super().__init__()
+        self.plugin_widgets = {}
 
         # Set the UI object
-        loader = QtUiTools.QUiLoader()
-        self.ui = loader.load(os.path.join(RELATIVE_PATH, "ui/settings.ui"))
+        self.ui = QtWidgets.QDialog()
+        self.ui.resize(600, 800)
+        dialog_layout = QtWidgets.QVBoxLayout(self.ui)
+        dialog_layout.setContentsMargins(2, 2, 2, 2)
+        
+        self.ui.searchBar = QtWidgets.QLineEdit()
+        self.ui.searchBar.setPlaceholderText("Search")
+        dialog_layout.addWidget(self.ui.searchBar)
 
-        self.settings_search = SettingsSearch([self.ui.scrollAreaWidgetContents], self.ui.scrollArea)
-
         self.ui.setWindowTitle("Settings")
-
         self.ui.finished.connect(self._on_finished)
 
-        max_port = (1 << 16) - 1
-        self.ui.osc_server_port_line_edit.setValidator(
-            QtGui.QIntValidator(0, max_port))
-        self.ui.osc_client_port_line_edit.setValidator(
-            QtGui.QIntValidator(0, max_port))
-
-        # Store the current config paths so we can warn about overwriting an
-        # existing config.
-        self._config_paths = config.list_config_paths()
-        self.set_config_path(SETTINGS.CONFIG)
-
+        self.general_settings_list = [
+            self._create_config_path_settings(),
+            self._create_switchboard_settings(),
+            self._create_project_settings(),
+            self._create_multi_user_server_settings()
+        ]
+        
+        self._create_tab_widget(dialog_layout)
+        
+        self.settings_search = SettingsSearch([self.ui.all_settings_scroll_area, self.ui.general_settings_scroll_area])
         self.ui.searchBar.textChanged.connect(
             self._on_search_text_edited)
+        
+    def select_all_tab(self):
+        self._on_tab_changed(0)
+        
+    def _on_tab_changed(self, index: int):
+        """
+        Because of the All category, widgets need to be re-parented constantly.
+        """
+        clear_widgets(self.ui.all_settings_scroll_area.layout())
+        clear_widgets(self.ui.general_settings_scroll_area.layout())
+        
+        is_all_category = index == 0
+        if is_all_category:
+            for general_setting_widget in self.general_settings_list:
+                self.ui.all_settings_scroll_area.layout().addWidget(general_setting_widget)
+               
+            for (group_box, scroll_bar) in self.plugin_widgets.values():
+                scroll_bar.takeWidget()
+                self.ui.all_settings_scroll_area.layout().addWidget(group_box)
+                
+            # Pull all widgets up so they do not attempt to divide space among themselves when filtered by search
+            self.ui.all_settings_scroll_area.layout().addItem(
+                QtWidgets.QSpacerItem(20, 40, QtWidgets.QSizePolicy.Minimum, QtWidgets.QSizePolicy.Expanding)
+            )
+        
+        # General settings may have been in All
+        is_general_category = index == 1
+        if is_general_category:
+            for general_setting_widget in self.general_settings_list:
+                self.ui.general_settings_scroll_area.layout().addWidget(general_setting_widget)
+            # Pull all widgets up so they do not attempt to divide space among themselves when filtered by search
+            self.ui.general_settings_scroll_area.layout().addItem(
+                QtWidgets.QSpacerItem(20, 40, QtWidgets.QSizePolicy.Minimum, QtWidgets.QSizePolicy.Expanding)
+            )
+           
+        # Add the plugin back to its tabs - it may have been in All
+        is_plugin = not is_all_category and not is_general_category 
+        if is_plugin:
+            plugin_name = self.ui.tab_widget.tabText(index)
+            (group_box, scroll_bar) = self.plugin_widgets[plugin_name]
+            is_plugin_widget_missing = scroll_bar.widget() is None
+            if is_plugin_widget_missing:
+                scroll_bar.setWidget(group_box)
+        
+        self._search_settings_again()
+        return
 
+    def _on_finished(self, result: int):
+        # Currently, the only way to dismiss the settings dialog is by using
+        # the close button as opposed to ok/cancel buttons, so we intercept the
+        # close to issue a warning if the config path was changed and we're
+        # about to overwrite some other existing config file.
+        if (self._changed_config_path in self._config_paths
+                and self._changed_config_path != self._current_config_path):
+            # Show the confirmation dialog using a relative path to the config.
+            rel_config_path = config.get_relative_config_path(
+                self._changed_config_path)
+            reply = QtWidgets.QMessageBox.question(
+                self.ui, 'Confirm Overwrite',
+                ('Are you sure you would like to change the config path and '
+                 f'overwrite the existing config file "{rel_config_path}"?'),
+                QtWidgets.QMessageBox.Yes, QtWidgets.QMessageBox.No)
+
+            if reply != QtWidgets.QMessageBox.Yes:
+                # Clear the changed config path so that the dialog returns the
+                # originally specified path when queried via config_path().
+                self._changed_config_path = None
+
+    def _search_settings_again(self):
+        self._on_search_text_edited(self.ui.searchBar.text())
+
+    def _on_search_text_edited(self, search_string: str):
+        self.settings_search.search(search_string)
+        
+    def _create_config_path_settings(self):
+        self.ui.config_path_layout = QtWidgets.QWidget()
+        layout = QtWidgets.QHBoxLayout(self.ui.config_path_layout)
+        
+        self.ui.config_path_label = QtWidgets.QLabel()
+        self.ui.config_path_label.setText("Config Path")
+        
+        self.ui.config_path_line_edit = QtWidgets.QLineEdit()
         self.ui.config_path_line_edit.textChanged.connect(
             self.config_path_text_changed)
+
+        layout.addWidget(self.ui.config_path_label)
+        layout.addWidget(self.ui.config_path_line_edit)
+
+        # Store the current config paths so we can warn about overwriting an existing config.
+        self._config_paths = config.list_config_paths()
+        self.set_config_path(SETTINGS.CONFIG)
+        
+        return self.ui.config_path_layout;
+        
+    def _create_switchboard_settings(self):
+        self.ui.switchboard_settings_group = QtWidgets.QGroupBox()
+        self.ui.switchboard_settings_group.setTitle("Switchboard")
+        layout = QtWidgets.QFormLayout(self.ui.switchboard_settings_group)
+        
+        # IP address
+        self.ui.ip_address_label = QtWidgets.QLabel()
+        self.ui.ip_address_label.setText("IP Address")
+        self.ui.ip_address_line_edit = QtWidgets.QLineEdit()
+        layout.addRow(self.ui.ip_address_label, self.ui.ip_address_line_edit)
+        
+        # Transport path
+        self.ui.transport_path_root = QtWidgets.QWidget()
+        self.ui.transport_path_label = QtWidgets.QLabel()
+        self.ui.transport_path_label.setText("Transport Path")
+        transport_path_horizontal_layout = QtWidgets.QHBoxLayout(self.ui.transport_path_root)
+        transport_path_horizontal_layout.setContentsMargins(0, 0, 0, 0)
+        self.ui.transport_path_line_edit = QtWidgets.QLineEdit()
+        self.ui.transport_path_browse_button = QtWidgets.QPushButton()
+        self.ui.transport_path_browse_button.setText("Browse")
+        transport_path_horizontal_layout.addWidget(self.ui.transport_path_line_edit)
+        transport_path_horizontal_layout.addWidget(self.ui.transport_path_browse_button)
+        layout.addRow(self.ui.transport_path_label, self.ui.transport_path_root)
+        
+        # Listener Executable name
+        self.ui.listener_exe_label = QtWidgets.QLabel()
+        self.ui.listener_exe_label.setText("Listener Executable name")
+        self.ui.listener_exe_line_edit = QtWidgets.QLineEdit()
+        layout.addRow(self.ui.listener_exe_label, self.ui.listener_exe_line_edit)
+        
+        return self.ui.switchboard_settings_group
+        
+    def _create_project_settings(self):
+        self.ui.project_settings_group = QtWidgets.QGroupBox()
+        self.ui.project_settings_group.setTitle("Project Settings")
+        layout = QtWidgets.QVBoxLayout(self.ui.project_settings_group)
+        
+        project_settings_root = QtWidgets.QWidget()
+        form_layout = QtWidgets.QFormLayout(project_settings_root)
+        layout.addWidget(project_settings_root)
+
+        # Project Name
+        self.ui.project_name_label = QtWidgets.QLabel()
+        self.ui.project_name_label.setText("Project Name")
+        self.ui.project_name_line_edit = QtWidgets.QLineEdit()
+        form_layout.addRow(self.ui.project_name_label, self.ui.project_name_line_edit)
+        
+        # UProject
+        self.ui.uproject_root = QtWidgets.QWidget()
+        self.ui.uproject_label = QtWidgets.QLabel()
+        self.ui.uproject_label.setText("UProject")
+        uproject_horizontal_layout = QtWidgets.QHBoxLayout(self.ui.uproject_root)
+        uproject_horizontal_layout.setContentsMargins(0, 0, 0, 0)
+        self.ui.uproject_line_edit = QtWidgets.QLineEdit()
+        self.ui.uproject_browse_button = QtWidgets.QPushButton()
+        self.ui.uproject_browse_button.setText("Browse")
+        uproject_horizontal_layout.addWidget(self.ui.uproject_line_edit)
+        uproject_horizontal_layout.addWidget(self.ui.uproject_browse_button)
+        form_layout.addRow(self.ui.uproject_label, self.ui.uproject_root)
+        
+        # Engine Dir
+        self.ui.engine_dir_root = QtWidgets.QWidget()
+        self.ui.engine_dir_label = QtWidgets.QLabel()
+        self.ui.engine_dir_label.setText("Engine Dir")
+        engine_dir_horizontal_layout = QtWidgets.QHBoxLayout(self.ui.engine_dir_root)
+        engine_dir_horizontal_layout.setContentsMargins(0, 0, 0, 0)
+        self.ui.engine_dir_line_edit = QtWidgets.QLineEdit()
+        self.ui.engine_dir_browse_button = QtWidgets.QPushButton()
+        self.ui.engine_dir_browse_button.setText("Browse")
+        engine_dir_horizontal_layout.addWidget(self.ui.engine_dir_line_edit)
+        engine_dir_horizontal_layout.addWidget(self.ui.engine_dir_browse_button)
+        form_layout.addRow(self.ui.engine_dir_label, self.ui.engine_dir_root)
+        
+        # Build Engine
+        self.ui.build_engine_label = QtWidgets.QLabel()
+        self.ui.build_engine_label.setText("Build Engine")
+        self.ui.build_engine_checkbox = QtWidgets.QCheckBox()
+        form_layout.addRow(self.ui.build_engine_label, self.ui.build_engine_checkbox)
+
+        # Map Path
+        self.ui.map_path_label = QtWidgets.QLabel()
+        self.ui.map_path_label.setText("Map Path")
+        self.ui.map_path_line_edit = QtWidgets.QLineEdit()
+        form_layout.addRow(self.ui.map_path_label, self.ui.map_path_line_edit)
+        
+        # Map Filter
+        self.ui.map_filter_label = QtWidgets.QLabel()
+        self.ui.map_filter_label.setText("Project Filter")
+        self.ui.map_filter_line_edit = QtWidgets.QLineEdit()
+        form_layout.addRow(self.ui.map_filter_label, self.ui.map_filter_line_edit)
+        
+        # Set up listeners
         self.ui.uproject_browse_button.clicked.connect(
             self.uproject_browse_button_clicked)
         self.ui.engine_dir_browse_button.clicked.connect(
             self.engine_dir_browse_button_clicked)
 
-        # update settings in CONFIG when they are changed in the SettingsDialog
+        # Update settings in CONFIG when they are changed in the SettingsDialog
         self.ui.engine_dir_line_edit.editingFinished.connect(
             lambda widget=self.ui.engine_dir_line_edit:
-                CONFIG.ENGINE_DIR.update_value(widget.text()))
+            CONFIG.ENGINE_DIR.update_value(widget.text()))
         self.ui.build_engine_checkbox.stateChanged.connect(
             lambda state:
-                CONFIG.BUILD_ENGINE.update_value(
-                    True if state == QtCore.Qt.Checked else False))
+            CONFIG.BUILD_ENGINE.update_value(
+                True if state == QtCore.Qt.Checked else False))
         self.ui.uproject_line_edit.editingFinished.connect(
             lambda widget=self.ui.uproject_line_edit:
-                CONFIG.UPROJECT_PATH.update_value(widget.text()))
+            CONFIG.UPROJECT_PATH.update_value(widget.text()))
         self.ui.map_path_line_edit.editingFinished.connect(
             lambda widget=self.ui.map_path_line_edit:
-                CONFIG.MAPS_PATH.update_value(widget.text()))
+            CONFIG.MAPS_PATH.update_value(widget.text()))
         self.ui.map_filter_line_edit.editingFinished.connect(
             lambda widget=self.ui.map_filter_line_edit:
-                CONFIG.MAPS_FILTER.update_value(widget.text()))
-
+            CONFIG.MAPS_FILTER.update_value(widget.text()))
+        
+        # Sub settings
+        self._create_osc_settings(layout)
+        self._create_source_control_settings(layout)
+        return self.ui.project_settings_group
+    
+    def _create_osc_settings(self, layout):
+        self.ui.osc_settings_group = QtWidgets.QGroupBox()
+        self.ui.osc_settings_group.setTitle("OSC")
+        form_layout = QtWidgets.QFormLayout(self.ui.osc_settings_group)
+        layout.addWidget(self.ui.osc_settings_group)
+        
+        # Server port
+        self.ui.osc_server_port_label = QtWidgets.QLabel()
+        self.ui.osc_server_port_label.setText("Server Port")
+        self.ui.osc_server_port_line_edit = QtWidgets.QLineEdit()
+        form_layout.addRow(self.ui.osc_server_port_label, self.ui.osc_server_port_line_edit)
+        
+        # Client port
+        self.ui.osc_client_port_label = QtWidgets.QLabel()
+        self.ui.osc_client_port_label.setText("Client Port")
+        self.ui.osc_client_port_line_edit = QtWidgets.QLineEdit()
+        form_layout.addRow(self.ui.osc_client_port_label, self.ui.osc_client_port_line_edit)
+        
+        # Set up listeners
+        max_port = (1 << 16) - 1
+        self.ui.osc_server_port_line_edit.setValidator(
+            QtGui.QIntValidator(0, max_port))
+        self.ui.osc_client_port_line_edit.setValidator(
+            QtGui.QIntValidator(0, max_port))
         self.ui.osc_server_port_line_edit.editingFinished.connect(
             lambda widget=self.ui.osc_server_port_line_edit:
-                CONFIG.OSC_SERVER_PORT.update_value(int(widget.text())))
+            CONFIG.OSC_SERVER_PORT.update_value(int(widget.text())))
         self.ui.osc_client_port_line_edit.editingFinished.connect(
             lambda widget=self.ui.osc_client_port_line_edit:
-                CONFIG.OSC_CLIENT_PORT.update_value(int(widget.text())))
-
+            CONFIG.OSC_CLIENT_PORT.update_value(int(widget.text())))
+    
+    def _create_source_control_settings(self, layout):
+        self.ui.source_control_settings_group = QtWidgets.QGroupBox()
+        self.ui.source_control_settings_group.setTitle("Source Control")
+        self.ui.source_control_settings_group.setCheckable(True)
+        form_layout = QtWidgets.QFormLayout(self.ui.source_control_settings_group)
+        layout.addWidget(self.ui.source_control_settings_group)
+        
+        # P4 Project Path
+        self.ui.p4_project_path_label = QtWidgets.QLabel()
+        self.ui.p4_project_path_label.setText("P4 Project Path")
+        self.ui.p4_project_path_line_edit = QtWidgets.QLineEdit()
+        form_layout.addRow(self.ui.p4_project_path_label, self.ui.p4_project_path_line_edit)
+        
+        # P4 Engine Path
+        self.ui.p4_engine_path_label = QtWidgets.QLabel()
+        self.ui.p4_engine_path_label.setText("P4 Engine Path")
+        self.ui.p4_engine_path_line_edit = QtWidgets.QLineEdit()
+        form_layout.addRow(self.ui.p4_engine_path_label, self.ui.p4_engine_path_line_edit)
+        
+        # P4 Workspace name
+        self.ui.source_control_workspace_label = QtWidgets.QLabel()
+        self.ui.source_control_workspace_label.setText("P4 Workspace name")
+        self.ui.source_control_workspace_line_edit = QtWidgets.QLineEdit()
+        form_layout.addRow(self.ui.source_control_workspace_label, self.ui.source_control_workspace_line_edit)
+        
+        # Set up listeners
         self.ui.p4_project_path_line_edit.editingFinished.connect(
             lambda widget=self.ui.p4_project_path_line_edit:
-                CONFIG.P4_PROJECT_PATH.update_value(widget.text()))
+            CONFIG.P4_PROJECT_PATH.update_value(widget.text()))
         self.ui.p4_engine_path_line_edit.editingFinished.connect(
             lambda widget=self.ui.p4_engine_path_line_edit:
-                CONFIG.P4_ENGINE_PATH.update_value(widget.text()))
+            CONFIG.P4_ENGINE_PATH.update_value(widget.text()))
         self.ui.source_control_workspace_line_edit.editingFinished.connect(
             lambda widget=self.ui.source_control_workspace_line_edit:
-                CONFIG.SOURCE_CONTROL_WORKSPACE.update_value(widget.text()))
+            CONFIG.SOURCE_CONTROL_WORKSPACE.update_value(widget.text()))
+    
+    def _create_multi_user_server_settings(self):
+        self.ui.multi_user_settings = QtWidgets.QGroupBox()
+        self.ui.multi_user_settings.setTitle("Multi User Server")
+        form_layout = QtWidgets.QFormLayout(self.ui.multi_user_settings)
+        
+        # Server Name
+        self.ui.mu_server_name_label = QtWidgets.QLabel()
+        self.ui.mu_server_name_label.setText("Server Name")
+        self.ui.mu_server_name_line_edit = QtWidgets.QLineEdit()
+        self.ui.mu_server_name_line_edit.setPlaceholderText("MU_Server")
+        form_layout.addRow(self.ui.mu_server_name_label, self.ui.mu_server_name_line_edit)
+        
+        # Command Line Args
+        self.ui.mu_cmd_line_args_label = QtWidgets.QLabel()
+        self.ui.mu_cmd_line_args_label.setText("Command Line Args")
+        self.ui.mu_cmd_line_args_line_edit = QtWidgets.QLineEdit()
+        form_layout.addRow(self.ui.mu_cmd_line_args_label, self.ui.mu_cmd_line_args_line_edit)
+        
+        # Unicast Endpoint
+        self.ui.mu_server_endpoint_label = QtWidgets.QLabel()
+        self.ui.mu_server_endpoint_label.setText("Unicast Endpoint")
+        self.ui.mu_server_endpoint_line_edit = QtWidgets.QLineEdit()
+        form_layout.addRow(self.ui.mu_server_endpoint_label, self.ui.mu_server_endpoint_line_edit)
+        
+        # Multicast Endpoint
+        self.ui.mu_server_multicast_endpoint_label = QtWidgets.QLabel()
+        self.ui.mu_server_multicast_endpoint_label.setText("Multicast Endpoint")
+        self.ui.mu_server_multicast_endpoint_line_edit = QtWidgets.QLineEdit()
+        form_layout.addRow(self.ui.mu_server_multicast_endpoint_label, self.ui.mu_server_multicast_endpoint_line_edit)
+        
+        # Clean History
+        self.ui.mu_clean_history_label = QtWidgets.QLabel()
+        self.ui.mu_clean_history_label.setText("Clean History")
+        self.ui.mu_clean_history_check_box = QtWidgets.QCheckBox()
+        form_layout.addRow(self.ui.mu_clean_history_label, self.ui.mu_clean_history_check_box)
+        
+        # Auto Launch
+        self.ui.mu_auto_launch_label = QtWidgets.QLabel()
+        self.ui.mu_auto_launch_label.setText("Auto Launch")
+        self.ui.mu_auto_launch_check_box = QtWidgets.QCheckBox()
+        form_layout.addRow(self.ui.mu_auto_launch_label, self.ui.mu_auto_launch_check_box)
+        
+        # Executable Name
+        self.ui.muserver_exe_label = QtWidgets.QLabel()
+        self.ui.muserver_exe_label.setText("Executable Name")
+        self.ui.muserver_exe_line_edit = QtWidgets.QLineEdit()
+        form_layout.addRow(self.ui.muserver_exe_label, self.ui.muserver_exe_line_edit)
+        
+        # Auto Build
+        self.ui.muserver_auto_build_label = QtWidgets.QLabel()
+        self.ui.muserver_auto_build_label.setText("Auto Build")
+        self.ui.muserver_auto_build_check_box = QtWidgets.QCheckBox()
+        form_layout.addRow(self.ui.muserver_auto_build_label, self.ui.muserver_auto_build_check_box)
+        
+        # Auto Static Endpoint
+        self.ui.muserver_auto_endpoint_label = QtWidgets.QLabel()
+        self.ui.muserver_auto_endpoint_label.setText("Auto Static Endpoint")
+        self.ui.muserver_auto_endpoint_check_box = QtWidgets.QCheckBox()
+        form_layout.addRow(self.ui.muserver_auto_endpoint_label, self.ui.muserver_auto_endpoint_check_box)
+        
+        return self.ui.multi_user_settings
 
-        self._device_groupbox = {}
+    def _create_tab_widget(self, parent_layout):
+        self.ui.tab_widget = HorizontalTabWidget()
+        self.ui.tab_widget.setTabPosition(QtWidgets.QTabWidget.West)
+        parent_layout.addWidget(self.ui.tab_widget)
+    
+        # All tab
+        self.ui.all_tab_root = QtWidgets.QScrollArea()
+        self.ui.all_tab_root.setWidgetResizable(True)
+        self.ui.all_tab_root.setSizeAdjustPolicy(QtWidgets.QAbstractScrollArea.AdjustToContents)
+        self.ui.all_tab_root.setSizePolicy(QtWidgets.QSizePolicy.Expanding, QtWidgets.QSizePolicy.Expanding)
 
-    def _on_finished(self, result: int):
-        # Currently, the only way to dismiss the settings dialog is by using
-        # the close button as opposed to ok/cancel buttons, so we intercept the
-        # close to issue a warning if the config path was changed and we're
-        # about to overwrite some other existing config file.
-        if (self._changed_config_path in self._config_paths 
-                and self._changed_config_path != self._current_config_path):
-            # Show the confirmation dialog using a relative path to the config.
-            rel_config_path = config.get_relative_config_path(
-                self._changed_config_path)
-            reply = QtWidgets.QMessageBox.question(
-                self.ui, 'Confirm Overwrite',
-                ('Are you sure you would like to change the config path and '
-                 f'overwrite the existing config file "{rel_config_path}"?'),
-                QtWidgets.QMessageBox.Yes, QtWidgets.QMessageBox.No)
+        self.ui.all_settings_scroll_area = QtWidgets.QWidget()
+        all_layout = QtWidgets.QVBoxLayout(self.ui.all_settings_scroll_area)
+        all_layout.setContentsMargins(2, 2, 2, 2)
+        self.ui.all_tab_root.setWidget(self.ui.all_settings_scroll_area)
+    
+        # General tab
+        self.ui.general_tab_root = QtWidgets.QScrollArea()
+        self.ui.general_tab_root.setWidgetResizable(True)
+        self.ui.general_tab_root.setSizeAdjustPolicy(QtWidgets.QAbstractScrollArea.AdjustToContents)
+        self.ui.general_tab_root.setSizePolicy(QtWidgets.QSizePolicy.Expanding, QtWidgets.QSizePolicy.Expanding)
+        
+        self.ui.general_settings_scroll_area = QtWidgets.QWidget()
+        general_layout = QtWidgets.QVBoxLayout(self.ui.general_settings_scroll_area)
+        general_layout.setContentsMargins(2, 2, 2, 2)
+        self.ui.general_tab_root.setWidget(self.ui.general_settings_scroll_area)
 
-            if reply != QtWidgets.QMessageBox.Yes:
-                # Clear the changed config path so that the dialog returns the
-                # originally specified path when queried via config_path().
-                self._changed_config_path = None
+        # Register tabs after all widgets have been allocated
+        self.ui.tab_widget.addTab(self.ui.all_tab_root, "All")
+        self.ui.tab_widget.addTab(self.ui.general_tab_root, "General")
+        # Register callback last because addTab should not trigger our callback before everything is initialized
+        self.ui.tab_widget.currentChanged.connect(
+            self._on_tab_changed
+        )
 
-    def _on_search_text_edited(self, search_string: str):
-        self.settings_search.search(search_string)
-
     def config_path(self):
         if self._changed_config_path:
@@ -204,10 +522,10 @@
         self.ui.build_engine_checkbox.setChecked(value)
 
     def p4_enabled(self):
-        return self.ui.p4_group_box.isChecked()
+        return self.ui.source_control_settings_group.isChecked()
 
     def set_p4_enabled(self, enabled):
-        self.ui.p4_group_box.setChecked(enabled)
+        self.ui.source_control_settings_group.setChecked(enabled)
 
     def p4_project_path(self):
         return self.ui.p4_project_path_line_edit.text()
@@ -302,8 +620,7 @@
         self.ui.muserver_auto_endpoint_check_box.setChecked(value)
 
     # Devices
-    def add_section_for_plugin(
-        self, plugin_name, plugin_settings, device_settings):
+    def add_section_for_plugin(self, plugin_name, plugin_settings, device_settings):
         any_device_settings = (
             any([device[1] for device in device_settings]) or
             any([device[2] for device in device_settings]))
@@ -311,14 +628,25 @@
             return  # no settings to show
 
         # Create a group box per plugin
-        device_override_group_box = self._device_groupbox.setdefault(
-            plugin_name, QtWidgets.QGroupBox())
-        if device_override_group_box.parent() is None:
-            device_override_group_box.setTitle(f'{plugin_name} Settings')
-            device_override_group_box.setLayout(QtWidgets.QVBoxLayout())
-            device_override_group_box.setSizePolicy(
-                QtWidgets.QSizePolicy(QtWidgets.QSizePolicy.MinimumExpanding, QtWidgets.QSizePolicy.Maximum))
-            self.ui.device_override_layout.addWidget(device_override_group_box)
+        device_override_group_box = QtWidgets.QGroupBox()
+        device_override_group_box.setTitle(f'{plugin_name} Settings')
+        device_override_group_box.setLayout(QtWidgets.QVBoxLayout())
+        device_override_group_box.setSizePolicy(
+            QtWidgets.QSizePolicy(QtWidgets.QSizePolicy.Expanding, QtWidgets.QSizePolicy.Expanding)
+        )
 
+        # Need to manually create scroll bar for tab
+        scroll_bar = QtWidgets.QScrollArea()
+        scroll_bar.setWidget(device_override_group_box)
+        scroll_bar.setWidgetResizable(True)
+        scroll_bar.setSizeAdjustPolicy(QtWidgets.QAbstractScrollArea.AdjustToContents)
+        scroll_bar.setSizePolicy(QtWidgets.QSizePolicy.Expanding, QtWidgets.QSizePolicy.Expanding)
+        scroll_layout = QtWidgets.QVBoxLayout()
+        scroll_layout.setContentsMargins(0, 0, 0, 0)
+        
+        self.ui.tab_widget.addTab(scroll_bar, f'{plugin_name}')
+        self.plugin_widgets[plugin_name] = (device_override_group_box, scroll_bar)
+        self.settings_search.searched_widgets.append(device_override_group_box)
+
         plugin_layout = QtWidgets.QFormLayout()
         device_override_group_box.layout().addLayout(plugin_layout)
@@ -345,3 +673,8 @@
             for setting in overrides:
                 setting.create_ui(
                     form_layout=layout, override_device_name=device_name)
+
+        # Pull all widgets up so they do not attempt to divide space among themselves when filtered by search
+        device_override_group_box.layout().addItem(
+            QtWidgets.QSpacerItem(20, 40, QtWidgets.QSizePolicy.Minimum, QtWidgets.QSizePolicy.Expanding)
+        )
diff --git a/Engine/Plugins/VirtualProduction/Switchboard/Source/Switchboard/switchboard/settings_search.py b/Engine/Plugins/VirtualProduction/Switchboard/Source/Switchboard/switchboard/settings_search.py
--- a/Engine/Plugins/VirtualProduction/Switchboard/Source/Switchboard/switchboard/settings_search.py
+++ b/Engine/Plugins/VirtualProduction/Switchboard/Source/Switchboard/switchboard/settings_search.py
@@ -118,45 +118,17 @@ class FormLayoutRowVisibility:
 
 
 class SettingsSearch:
-    def __init__(self, searched_widgets: List[QtWidgets.QWidget], scroll_area: QtWidgets.QScrollArea):
-        self.reset_search_bar_height = 0
-        self.last_search_string_len = 0
-        
+    def __init__(self, searched_widgets: List[QtWidgets.QWidget]):
         self.searched_widgets = searched_widgets
-        self.scroll_area = scroll_area
         
         form_layout_handlers: Dict[QtWidgets.QFormLayout, FormLayoutRowVisibility] = {}
         self.form_layout_handlers = form_layout_handlers
         
     def search(self, search_string: str):
-        # The first time something is typed, save scroll height to restore later
-        self._take_scroll_snapshot(search_string)
-        
         search_term_list = search_string.split()
         for widget in self.searched_widgets:
             self._update_widget_visibility(widget, search_term_list)
 
-        # Restore scroll bar height when search is over
-        self._restore_scroll_snapshot(search_string)
-
-    def _take_scroll_snapshot(self, search_string: str):
-        just_started_search = self.last_search_string_len == 0
-        if just_started_search:
-            self.reset_search_bar_height = self.scroll_area.verticalScrollBar().value()
-            
-            self.scroll_area.verticalScrollBar().setValue(0)
-
-    def _restore_scroll_snapshot(self, search_string: str):
-        len_search = len(search_string)
-        is_search_empty = len_search < self.last_search_string_len and len_search == 0
-        if is_search_empty:
-            # Need to wait for UI to update layout...
-            QtCore.QTimer.singleShot(
-                25,
-                lambda height=self.reset_search_bar_height: self.scroll_area.verticalScrollBar().setValue(height)
-            )
-        self.last_search_string_len = len_search
-
     def _update_widget_visibility(self, widget: QtWidgets.QWidget, search_term_list):
         # We only search labels for search terms
         if isinstance(widget, QtWidgets.QLabel):
diff --git a/Engine/Plugins/VirtualProduction/Switchboard/Source/Switchboard/switchboard/switchboard_dialog.py b/Engine/Plugins/VirtualProduction/Switchboard/Source/Switchboard/switchboard/switchboard_dialog.py
--- a/Engine/Plugins/VirtualProduction/Switchboard/Source/Switchboard/switchboard/switchboard_dialog.py
+++ b/Engine/Plugins/VirtualProduction/Switchboard/Source/Switchboard/switchboard/switchboard_dialog.py
@@ -840,6 +840,8 @@ class SwitchboardDialog(QtCore.QObject):
             device_instances = self.device_manager.devices_of_type(plugin_name)
             device_settings = [(device.name, device.device_settings(), device.setting_overrides()) for device in device_instances]
             settings_dialog.add_section_for_plugin(plugin_name, self.device_manager.plugin_settings(plugin_name), device_settings)
+        
+        settings_dialog.select_all_tab()
 
         # avoid saving the config all the time while in the settings dialog
         CONFIG.push_saving_allowed(False)
diff --git a/Engine/Plugins/VirtualProduction/Switchboard/Source/Switchboard/switchboard/ui/horizontal_tabs.py b/Engine/Plugins/VirtualProduction/Switchboard/Source/Switchboard/switchboard/ui/horizontal_tabs.py
new file mode 100644
--- /dev/null
+++ b/Engine/Plugins/VirtualProduction/Switchboard/Source/Switchboard/switchboard/ui/horizontal_tabs.py
@@ -0,0 +1,29 @@
+# Copyright Epic Games, Inc. All Rights Reserved.
+
+from PySide2 import QtCore
+from PySide2 import QtWidgets
+
+
+class HorizontalTabBar(QtWidgets.QTabBar):
+    def paintEvent(self, event):
+        painter = QtWidgets.QStylePainter(self)
+        option = QtWidgets.QStyleOptionTab()
+        for index in range(self.count()):
+            self.initStyleOption(option, index)
+            painter.drawControl(QtWidgets.QStyle.CE_TabBarTabShape, option)
+            painter.drawText(self.tabRect(index),
+                             QtCore.Qt.AlignCenter | QtCore.Qt.TextDontClip,
+                             self.tabText(index))
+
+    def tabSizeHint(self, index):
+        size = QtWidgets.QTabBar.tabSizeHint(self, index)
+        if size.width() < size.height():
+            size.transpose()
+        size.setHeight(30)
+        return size
+
+
+class HorizontalTabWidget(QtWidgets.QTabWidget):
+    def __init__(self, parent=None):
+        QtWidgets.QTabWidget.__init__(self, parent)
+        self.setTabBar(HorizontalTabBar())
\ No newline at end of file
