Improve logging settings: multi-select categories

diff --git a/Engine/Plugins/VirtualProduction/Switchboard/Source/Switchboard/switchboard/config.py b/Engine/Plugins/VirtualProduction/Switchboard/Source/Switchboard/switchboard/config.py
--- a/Engine/Plugins/VirtualProduction/Switchboard/Source/Switchboard/switchboard/config.py
+++ b/Engine/Plugins/VirtualProduction/Switchboard/Source/Switchboard/switchboard/config.py
@@ -788,6 +788,9 @@ class LoggingModel(QtGui.QStandardItemModel):
 
         return True
 
+    def category_at(self, index: QtCore.QModelIndex):
+        return self.invisibleRootItem().child(index.row(), LoggingModel.CATEGORY_COLUMN).text()
+
     @property
     def category_verbosities(self) -> collections.OrderedDict:
         '''
@@ -822,11 +825,9 @@ class LoggingModel(QtGui.QStandardItemModel):
         value = collections.OrderedDict(sorted(value.items(), key=lambda item: str(item[0]).lower()))
 
         self.beginResetModel()
-
-        self.removeRows(0, self.rowCount())
-
+        count_before = self.rowCount()
+        
         root_item = self.invisibleRootItem()
-
         for category, verbosity_level in value.items():
             if (category not in self._categories and
                     category not in self._user_categories):
@@ -835,7 +836,10 @@ class LoggingModel(QtGui.QStandardItemModel):
             root_item.appendRow(
                 [QtGui.QStandardItem(category),
                     QtGui.QStandardItem(verbosity_level)])
-
+            
+        # This triggers the rowsRemoved event. 
+        # Remove rows after so external observers get the right result from category_verbosities.  
+        self.removeRows(0, count_before)
         self.endResetModel()
 
     def flags(self, index: QtCore.QModelIndex) -> QtCore.Qt.ItemFlags:
@@ -864,9 +868,11 @@ class LoggingVerbosityItemDelegate(QtWidgets.QStyledItemDelegate):
     def __init__(
             self,
             verbosity_levels: typing.List[str],
-            parent: typing.Optional[QtCore.QObject] = None):
+            parent: QtWidgets.QTreeView):
         super().__init__(parent=parent)
         self._verbosity_levels = verbosity_levels or []
+        self._parent = parent
+        self._selected_categories = []
 
     def createEditor(
             self, parent: QtWidgets.QWidget,
@@ -881,8 +887,18 @@ class LoggingVerbosityItemDelegate(QtWidgets.QStyledItemDelegate):
             current_index = self._verbosity_levels.index(current_value)
             editor.setCurrentIndex(current_index)
 
-        editor.currentIndexChanged.connect(self.currentIndexChanged)
-
+        edited_category = self._parent.model().category_at(index)
+        editor.onHoverScrollBox.connect(
+            lambda edited_category=edited_category:
+                self.on_hover_combo_box(edited_category)
+        )
+        # For multi-selection we want to also be modified even when index has not changed
+        editor.activated.connect(
+            lambda combo_index, 
+                editor=editor:
+                self.on_current_index_changed(combo_index, editor)
+        )
+        
         return editor
 
     def setEditorData(
@@ -896,8 +912,13 @@ class LoggingVerbosityItemDelegate(QtWidgets.QStyledItemDelegate):
             index: QtCore.QModelIndex):
         model.setData(index, editor.currentText() or None)
 
-    def currentIndexChanged(self):
-        self.commitData.emit(self.sender())
+    def on_hover_combo_box(self, edited_category: str):
+        selection = self._parent.selected_categories()
+        # Discard the selection if the user clicked a combo box outside of the selection
+        self._selected_categories = selection if edited_category in selection else [edited_category]
+
+    def on_current_index_changed(self, combo_index, editor):
+        self._parent.update_category_verbosities(self._selected_categories, editor.itemText(combo_index))
 
 
 class LoggingSettingVerbosityView(QtWidgets.QTreeView):
@@ -925,14 +946,45 @@ class LoggingSettingVerbosityView(QtWidgets.QTreeView):
         self.setItemDelegateForColumn(
             LoggingModel.VERBOSITY_COLUMN,
             LoggingVerbosityItemDelegate(logging_model.verbosity_levels, self))
-
-        for row in range(logging_model.rowCount()):
-            verbosity_level_index = logging_model.index(
-                row, LoggingModel.VERBOSITY_COLUMN)
-            self.openPersistentEditor(verbosity_level_index)
+        
+        self._open_persistent_editors()
+        logging_model.modelAboutToBeReset.connect(self._pre_change_model)
+        logging_model.modelReset.connect(self._post_change_model)
 
         self.setSelectionBehavior(QtWidgets.QTreeView.SelectRows)
-
+        self.setSelectionMode(QtWidgets.QTreeView.ExtendedSelection)
+        
+    def _pre_change_model(self):
+        self.scroll_height = self.verticalScrollBar().value()
+        
+    def _post_change_model(self):
+        self.verticalScrollBar().setSliderPosition(self.scroll_height)
+        self._open_persistent_editors()
+        
+    def _open_persistent_editors(self):
+        for row in range(self.model().rowCount()):
+            verbosity_level_index = self.model().index(row, LoggingModel.VERBOSITY_COLUMN)
+            self.openPersistentEditor(verbosity_level_index)
+        
+    def selected_categories(self) -> typing.List[str]:
+        model = self.model()
+        selected_categories = []
+        for selection_range in self.selectionModel().selection():
+            rows = range(selection_range.top(), selection_range.bottom() + 1)
+            for row in rows:
+                category_index = model.index(row, 0)
+                category = model.data(category_index, QtCore.Qt.DisplayRole)
+                selected_categories.append(category)
+                
+        return selected_categories
+        
+    def update_category_verbosities(self, categories: typing.List[str], new_verbosity: str):
+        model = self.model()
+        category_verbosities = model.category_verbosities
+        for category in categories:
+            category_verbosities[category] = new_verbosity
+
+        model.category_verbosities = category_verbosities
 
 class LoggingSetting(Setting):
     '''
@@ -1047,7 +1099,7 @@ class LoggingSetting(Setting):
             verbosity_levels=self._verbosity_levels)
         model.category_verbosities = self._value
         view = LoggingSettingVerbosityView(logging_model=model)
-        view.setMinimumHeight(100)
+        view.setMinimumHeight(150)
 
         self.set_widget(
             widget=view, override_device_name=override_device_name)
@@ -1084,17 +1136,23 @@ class LoggingSetting(Setting):
                 LoggingModel.CATEGORY_COLUMN)
             if not category_indices:
                 return
-            category = model.itemFromIndex(category_indices[0]).text()
-
+            
+            category_list_str = ''
+            category_list = []
+            for category_index in category_indices:
+                category = model.itemFromIndex(category_index).text()
+                category_list.append(category)
+                category_list_str += f'{category}\n'
             reply = QtWidgets.QMessageBox.question(
                 remove_category_button, 'Confirm Remove Category',
-                ('Are you sure you would like to remove the logging category '
-                 f'"{category}"?'),
+                ('Are you sure you would like to remove the following categories: '
+                 f'{category_list_str}'),
                 QtWidgets.QMessageBox.Yes, QtWidgets.QMessageBox.No)
 
             if reply == QtWidgets.QMessageBox.Yes:
                 view.selectionModel().clear()
-                model.remove_user_category(category)
+                for category in category_list:
+                    model.remove_user_category(category)
 
         remove_category_button.clicked.connect(
             on_remove_category_button_clicked)
@@ -1117,15 +1175,17 @@ class LoggingSetting(Setting):
             if not category_indices:
                 return
 
-            category = model.itemFromIndex(category_indices[0]).text()
-
-            if model.is_user_category(category):
+            all_user_category = True
+            for category_index in category_indices:
+                category = model.itemFromIndex(category_index).text()
+                all_user_category &= model.is_user_category(category)
+            if all_user_category:
                 remove_category_button.setEnabled(True)
                 remove_category_button.setToolTip('')
 
         view.selectionModel().selectionChanged.connect(
             on_view_selectionChanged)
-
+        
         def on_logging_model_modified(override_device_name=None):
             category_verbosities = model.category_verbosities
             self._on_widget_value_changed(
@@ -1149,7 +1209,8 @@ class LoggingSetting(Setting):
             override_device_name=override_device_name:
                 on_logging_model_modified(
                     override_device_name=override_device_name))
-
+        
+        
         return edit_layout
 
     def _on_setting_changed(
diff --git a/Engine/Plugins/VirtualProduction/Switchboard/Source/Switchboard/switchboard/switchboard_widgets.py b/Engine/Plugins/VirtualProduction/Switchboard/Source/Switchboard/switchboard/switchboard_widgets.py
--- a/Engine/Plugins/VirtualProduction/Switchboard/Source/Switchboard/switchboard/switchboard_widgets.py
+++ b/Engine/Plugins/VirtualProduction/Switchboard/Source/Switchboard/switchboard/switchboard_widgets.py
@@ -13,9 +13,15 @@ DEVICE_WIDGET_HIDE_IP_ADDRESS_WIDTH = 500
 
 
 class NonScrollableComboBox(QtWidgets.QComboBox):
+    onHoverScrollBox = QtCore.Signal()
+    
     def __init__(self, parent=None):
         super().__init__(parent=parent)
         self.installEventFilter(self)
+        
+    def enterEvent(self, event):
+        self.onHoverScrollBox.emit()
+        super().enterEvent(event)
 
     def eventFilter(self, obj, event):
         if obj == self and event.type() == QtCore.QEvent.Wheel:
