Switchboard: Replaced master with primary.

This is also in preparation to the new nDisplay config file changes for 5.0.

diff --git a/Engine/Plugins/VirtualProduction/Switchboard/Source/Switchboard/switchboard/devices/ndisplay/ndisplay_monitor.py b/Engine/Plugins/VirtualProduction/Switchboard/Source/Switchboard/switchboard/devices/ndisplay/ndisplay_monitor.py
--- a/Engine/Plugins/VirtualProduction/Switchboard/Source/Switchboard/switchboard/devices/ndisplay/ndisplay_monitor.py
+++ b/Engine/Plugins/VirtualProduction/Switchboard/Source/Switchboard/switchboard/devices/ndisplay/ndisplay_monitor.py
@@ -606,7 +606,7 @@ class nDisplayMonitor(QAbstractTableModel):
 
     @QtCore.Slot()
     def on_soft_kill_clicked(self):
-        ''' Kills the cluster by sending a message to the master. '''
+        ''' Kills the cluster by sending a message to the primary. '''
         devices = [devicedata['device']
                    for devicedata in self.devicedatas.values()]
 
diff --git a/Engine/Plugins/VirtualProduction/Switchboard/Source/Switchboard/switchboard/devices/ndisplay/plugin_ndisplay.py b/Engine/Plugins/VirtualProduction/Switchboard/Source/Switchboard/switchboard/devices/ndisplay/plugin_ndisplay.py
--- a/Engine/Plugins/VirtualProduction/Switchboard/Source/Switchboard/switchboard/devices/ndisplay/plugin_ndisplay.py
+++ b/Engine/Plugins/VirtualProduction/Switchboard/Source/Switchboard/switchboard/devices/ndisplay/plugin_ndisplay.py
@@ -260,11 +260,10 @@ class AddnDisplayDialog(AddDeviceDialog):
                     f"{cfg_file}")
 
             # Initialize from the config our setting that identifies which
-            # device is the master.
+            # device is the primary.
             for node in devices:
-                if node['master']:
-                    DevicenDisplay.csettings[
-                        'master_device_name'].update_value(node['name'])
+                if node['primary']:
+                    DevicenDisplay.csettings['primary_device_name'].update_value(node['name'])
                     break
 
             return devices
@@ -279,7 +278,7 @@ class AddnDisplayDialog(AddDeviceDialog):
 
 class DeviceWidgetnDisplay(DeviceWidgetUnreal):
 
-    signal_device_widget_master = QtCore.Signal(object)
+    signal_device_widget_primary = QtCore.Signal(object)
 
     def _add_control_buttons(self):
         self._autojoin_visible = False
@@ -289,29 +288,29 @@ class DeviceWidgetnDisplay(DeviceWidgetUnreal):
         ''' DeviceWidget base class method override. '''
 
         if widget == self.name_line_edit:
-            self.add_master_button()
+            self.add_primary_button()
 
-            # shorten the widget to account for the inserted master button
+            # shorten the widget to account for the inserted primary button
             btn_added_width = (
                 max(
-                    self.master_button.iconSize().width(),
-                    self.master_button.minimumSize().width())
+                    self.primary_button.iconSize().width(),
+                    self.primary_button.minimumSize().width())
                 + 2 * self.layout.spacing()
-                + self.master_button.contentsMargins().left()
-                + self.master_button.contentsMargins().right())
+                + self.primary_button.contentsMargins().left()
+                + self.primary_button.contentsMargins().right())
 
             le_maxwidth = self.name_line_edit.maximumWidth() - btn_added_width
             self.name_line_edit.setMaximumWidth(le_maxwidth)
 
         super().add_widget_to_layout(widget)
 
-    def add_master_button(self):
+    def add_primary_button(self):
         '''
         Adds to the layout a button to select which device should be the
-        master device in the cluster.
+        primary device in the cluster.
         '''
 
-        self.master_button = sb_widgets.ControlQPushButton.create(
+        self.primary_button = sb_widgets.ControlQPushButton.create(
             ':/icons/images/star_yellow_off.png',
             icon_disabled=':/icons/images/star_yellow_off.png',
             icon_hover=':/icons/images/star_yellow_off.png',
@@ -319,17 +318,17 @@ class DeviceWidgetnDisplay(DeviceWidgetUnreal):
             icon_on=':/icons/images/star_yellow.png',
             icon_hover_on=':/icons/images/star_yellow.png',
             icon_size=QtCore.QSize(16, 16),
-            tool_tip='Select as master node'
+            tool_tip='Select as primary node'
         )
 
-        self.add_widget_to_layout(self.master_button)
+        self.add_widget_to_layout(self.primary_button)
 
-        self.master_button.clicked.connect(self.on_master_button_clicked)
+        self.primary_button.clicked.connect(self.on_primary_button_clicked)
 
-    def on_master_button_clicked(self):
-        ''' Called when master_button is clicked '''
+    def on_primary_button_clicked(self):
+        ''' Called when primary_button is clicked '''
 
-        self.signal_device_widget_master.emit(self)
+        self.signal_device_widget_primary.emit(self)
 
 
 class DisplayConfig(object):
@@ -448,13 +447,12 @@ class DevicenDisplay(DeviceUnreal):
             value=True,
             tool_tip="Minimizes windows before launch"
         ),
-        'master_device_name': StringSetting(
-            attr_name='master_device_name',
-            nice_name="Master Device",
+        'primary_device_name': StringSetting(
+            attr_name='primary_device_name',
+            nice_name="Primary Device",
             value='',
             tool_tip=(
-                "Identifies which nDisplay device should be the master in the "
-                "cluster"),
+                "Identifies which nDisplay device should be the primary node in the cluster"),
             show_ui=False,
         ),
         'logging': LoggingSetting(
@@ -975,20 +973,20 @@ class DevicenDisplay(DeviceUnreal):
 
         super().device_widget_registered(device_widget)
 
-        # hook to its master_button clicked event
-        device_widget.signal_device_widget_master.connect(
-            self.select_as_master)
+        # hook to its primary_button clicked event
+        device_widget.signal_device_widget_primary.connect(
+            self.select_as_primary)
 
         # Update the state of the button depending on whether this device is
-        # the master or not.
-        is_master = (
+        # the primary or not.
+        is_primary = (
             self.name ==
-            DevicenDisplay.csettings['master_device_name'].get_value())
-        device_widget.master_button.setChecked(is_master)
+            DevicenDisplay.csettings['primary_device_name'].get_value())
+        device_widget.primary_button.setChecked(is_primary)
 
-    def select_as_master(self):
-        ''' Selects this node as the master in the nDisplay cluster '''
-        self.__class__.select_device_as_master(self)
+    def select_as_primary(self):
+        ''' Selects this node as the primary node in the nDisplay cluster '''
+        self.__class__.select_device_as_primary(self)
 
     @classmethod
     def extract_configexport_from_uasset(cls, cfg_file) -> str:
@@ -1007,10 +1005,10 @@ class DevicenDisplay(DeviceUnreal):
         raise ValueError('Invalid nDisplay config .uasset')
 
     @classmethod
-    def select_device_as_master(cls, device):
-        ''' Selects the given devices as the master of the nDisplay cluster'''
+    def select_device_as_primary(cls, device):
+        ''' Selects the given devices as the primary of the nDisplay cluster'''
 
-        DevicenDisplay.csettings['master_device_name'].update_value(
+        DevicenDisplay.csettings['primary_device_name'].update_value(
             device.name)
 
         devices = [
@@ -1020,7 +1018,7 @@ class DevicenDisplay(DeviceUnreal):
         for dev in devices:
             if dev.widget:
                 checked = (dev.name == device.name)
-                dev.widget.master_button.setChecked(checked)
+                dev.widget.primary_button.setChecked(checked)
 
     @classmethod
     def apply_local_overrides_to_config(cls, cfg_content, cfg_ext):
@@ -1031,11 +1029,15 @@ class DevicenDisplay(DeviceUnreal):
 
         data = json.loads(cfg_content)
 
-        # override master node
+        configversion = float(data['nDisplay']['version'])
+        primaryNodeKey = 'masterNode' if configversion < 5.0 else 'primaryNode'
+
+        # override primary node
         #
 
         nodes = data['nDisplay']['cluster']['nodes']
-        masternodeid = data['nDisplay']['cluster']['masterNode']['id']
+
+        primaryNodeId = data['nDisplay']['cluster'][primaryNodeKey]['id']
 
         activenodes = {}
 
@@ -1066,14 +1068,12 @@ class DevicenDisplay(DeviceUnreal):
             # add to active nodes
             activenodes[device.name] = node
 
-            # override master node by name
-            if (device.name ==
-                    DevicenDisplay.csettings[
-                        'master_device_name'].get_value()):
-                masternodeid = device.name
+            # override primary node by name
+            if (device.name == DevicenDisplay.csettings['primary_device_name'].get_value()):
+                primaryNodeId = device.name
 
         data['nDisplay']['cluster']['nodes'] = activenodes
-        data['nDisplay']['cluster']['masterNode']['id'] = masternodeid
+        data['nDisplay']['cluster'][primaryNodeKey]['id'] = primaryNodeId
 
         # override sync policy
         #
@@ -1120,8 +1120,11 @@ class DevicenDisplay(DeviceUnreal):
 
         config = DisplayConfig()
 
+        configversion = float(js['nDisplay']['version'])
+        primaryNodeKey = 'masterNode' if configversion < 5.0 else 'primaryNode'
+
         cnodes = js['nDisplay']['cluster']['nodes']
-        masterNode = js['nDisplay']['cluster']['masterNode']
+        primaryNode = js['nDisplay']['cluster'][primaryNodeKey]
         config.uasset_path = js['nDisplay'].get('assetPath')
 
         for name, cnode in cnodes.items():
@@ -1137,13 +1140,13 @@ class DevicenDisplay(DeviceUnreal):
             # Note the capital 'S'.
             kwargs["fullscreen"] = bool(cnode.get('fullScreen', False))
         
-            master = True if masterNode['id'] == name else False
+            primary = True if primaryNode['id'] == name else False
 
             config.nodes.append({
                 "name": name,
                 "ip_address": cnode['host'],
-                "master": master,
-                "port_ce": int(masterNode['ports']['ClusterEventsJson']),
+                "primary": primary,
+                "port_ce": int(primaryNode['ports']['ClusterEventsJson']),
                 "kwargs": kwargs,
             })
 
@@ -1345,37 +1348,37 @@ class DevicenDisplay(DeviceUnreal):
     @classmethod
     def send_cluster_event(cls, devices, cluster_event):
         '''
-        Sends a cluster event (to the master, which will replicate to the rest
+        Sends a cluster event (to the primary node, which will replicate to the rest
         of the cluster).
         '''
 
-        # find the master node
+        # find the primary node
 
-        master = None
+        primary = None
 
         for device in devices:
             if (device.name ==
                     DevicenDisplay.csettings[
-                        'master_device_name'].get_value()):
-                master = device
+                        'primary_device_name'].get_value()):
+                primary = device
                 break
 
-        if master is None:
-            LOGGER.warning('Could not find master device when trying to send '
-                           'cluster event. Please make sure the master '
+        if primary is None:
+            LOGGER.warning('Could not find primary device when trying to send '
+                           'cluster event. Please make sure the primary '
                            'device is marked as such.')
             raise ValueError
 
         msg = bytes(json.dumps(cluster_event), 'utf-8')
         msg = struct.pack('I', len(msg)) + msg
 
-        # We copy the original master's cluster event port (port_ce) to the
+        # We copy the original primary's cluster event port (port_ce) to the
         # config of all nodes, so we can use the port from the overridden
-        # master.
-        port = master.nodeconfig['port_ce']
+        # primary.
+        port = primary.nodeconfig['port_ce']
 
         # Use the overridden ip address of this node
-        ip = master.ip_address
+        ip = primary.ip_address
 
         sock = socket.socket(socket.AF_INET, socket.SOCK_STREAM)
         sock.connect((ip, port))
@@ -1383,7 +1386,7 @@ class DevicenDisplay(DeviceUnreal):
 
     @classmethod
     def soft_kill_cluster(cls, devices):
-        ''' Kills the cluster by sending a message to the master. '''
+        ''' Kills the cluster by sending a message to the primary. '''
         LOGGER.info('Issuing nDisplay cluster soft kill')
 
         quit_event = {
@@ -1438,18 +1441,18 @@ class DevicenDisplay(DeviceUnreal):
         if len(devices) == 0:
             return
 
-        # Verify that there is a device that matches our master selection
-        master_device_name = DevicenDisplay.csettings['master_device_name'].get_value()
-        master_found = False
+        # Verify that there is a device that matches our primary selection
+        primary_device_name = DevicenDisplay.csettings['primary_device_name'].get_value()
+        primary_found = False
 
         for device in devices:
-            if device.name == master_device_name:
-                master_found = True
+            if device.name == primary_device_name:
+                primary_found = True
                 break
 
-        # If there isn't a master selection (e.g. old config), try to assign
+        # If there isn't a primary selection (e.g. old config), try to assign
         # one based on the nDisplay config.
-        if not master_found:
+        if not primary_found:
 
             cfg_file = DevicenDisplay.csettings['ndisplay_config_file'].get_value()
 
@@ -1460,15 +1463,15 @@ class DevicenDisplay(DeviceUnreal):
                 return
 
             for node in config.nodes:
-                if node['master']:
+                if node['primary']:
                     for device in devices:
                         if node['name'] == device.name:
-                            cls.select_device_as_master(device)
-                            master_found = True
+                            cls.select_device_as_primary(device)
+                            primary_found = True
                             break
                     break
 
-            # If we still don't have a master, log the error.
-            if not master_found:
-                LOGGER.error('Could not find or assign master device in the '
+            # If we still don't have a primary device, log the error.
+            if not primary_found:
+                LOGGER.error('Could not find or assign primary device in the '
                              f'current devices and config file "{cfg_file}"')
