Fix combo box with only intergers not getting reset

diff --git a/Engine/Plugins/VirtualProduction/Switchboard/Source/Switchboard/switchboard/config.py b/Engine/Plugins/VirtualProduction/Switchboard/Source/Switchboard/switchboard/config.py
--- a/Engine/Plugins/VirtualProduction/Switchboard/Source/Switchboard/switchboard/config.py
+++ b/Engine/Plugins/VirtualProduction/Switchboard/Source/Switchboard/switchboard/config.py
@@ -693,8 +693,9 @@ class OptionSetting(Setting):
             return
 
         old_value = widget.currentText()
-        if new_value != old_value:
-            widget.setCurrentIndex(widget.findText(new_value))
+        new_str_value = new_value if isinstance(new_value, str) else str(new_value)
+        if new_str_value != old_value:
+            widget.setCurrentIndex(widget.findText(new_str_value))
 
 
 class MultiOptionSetting(OptionSetting):
