Migrate properties from string to list

diff --git a/Engine/Plugins/VirtualProduction/Switchboard/Source/Switchboard/switchboard/config.py b/Engine/Plugins/VirtualProduction/Switchboard/Source/Switchboard/switchboard/config.py
--- a/Engine/Plugins/VirtualProduction/Switchboard/Source/Switchboard/switchboard/config.py
+++ b/Engine/Plugins/VirtualProduction/Switchboard/Source/Switchboard/switchboard/config.py
@@ -23,6 +23,13 @@ USER_SETTINGS_FILE_PATH = ROOT_CONFIGS_PATH.joinpath('user_settings.json')
 
 DEFAULT_MAP_TEXT = '-- Default Map --'
 
+def migrate_comma_separated_string_to_list(value) -> typing.List[str]:
+    if isinstance(value, str):
+        return value.split(",")
+    # Technically we should check whether every element is a string but we skip it here
+    if isinstance(value, list):
+        return value
+    raise NotImplementedError("Migration not handled")
 
 class Setting(QtCore.QObject):
     '''
@@ -52,7 +59,8 @@ class Setting(QtCore.QObject):
         value,
         tool_tip: typing.Optional[str] = None,
         show_ui: bool = True,
-        allow_reset: bool = True
+        allow_reset: bool = True,
+        migrate_data: typing.Callable[[any], None] = None
     ):
         '''
         Create a new Setting object.
@@ -93,6 +101,8 @@ class Setting(QtCore.QObject):
         self._base_reset_widget = None
         self._reset_override_widgets = {}
 
+        self._migrate_data = migrate_data
+
         self.tool_tip = tool_tip
         self.show_ui = show_ui
 
@@ -131,10 +141,15 @@ class Setting(QtCore.QObject):
         self._refresh_reset_override_widget(device_name)
 
     def get_value(self, device_name: typing.Optional[str] = None):
-        try:
-            return self._overrides[device_name]
-        except KeyError:
-            return self._value
+        def get_value(self, device_name: typing.Optional[str] = None):
+            try:
+                return self._overrides[device_name]
+            except KeyError:
+                return self._value
+            
+        value = get_value(self, device_name)
+        return self._migrate_data(value) if self._migrate_data is not None else value
+        
 
     def on_device_name_changed(self, old_name: str, new_name: str):
         if old_name in self._overrides.keys():
@@ -451,7 +466,8 @@ class StringSetting(Setting):
         placeholder_text: str = '',
         tool_tip: typing.Optional[str] = None,
         show_ui: bool = True,
-        allow_reset: bool = True
+        allow_reset: bool = True,
+        migrate_data: typing.Callable[[any], None] = None
     ):
         '''
         Create a new StringSetting object.
@@ -466,7 +482,7 @@ class StringSetting(Setting):
         '''
         super().__init__(
             attr_name, nice_name, value,
-            tool_tip=tool_tip, show_ui=show_ui, allow_reset=allow_reset)
+            tool_tip=tool_tip, show_ui=show_ui, allow_reset=allow_reset, migrate_data=migrate_data)
 
         self.placeholder_text = placeholder_text
 
@@ -644,7 +660,8 @@ class OptionSetting(Setting):
         possible_values: typing.List = None,
         tool_tip: typing.Optional[str] = None,
         show_ui: bool = True,
-        allow_reset: bool = True
+        allow_reset: bool = True,
+        migrate_data: typing.Callable[[any], None] = None
     ):
         '''
         Create a new OptionSetting object.
@@ -659,7 +676,7 @@ class OptionSetting(Setting):
         '''
         super().__init__(
             attr_name, nice_name, value,
-            tool_tip=tool_tip, show_ui=show_ui, allow_reset=allow_reset)
+            tool_tip=tool_tip, show_ui=show_ui, allow_reset=allow_reset, migrate_data=migrate_data)
 
         self.possible_values = possible_values or []
 
@@ -794,7 +811,8 @@ class ListSetting(Setting):
         value: typing.List = [],
         tool_tip: typing.Optional[str] = None,
         show_ui: bool = True,
-        allow_reset: bool = True
+        allow_reset: bool = True,
+        migrate_data: typing.Callable[[any], None] = None
     ):
         '''
         Create a new ArraySetting object.
@@ -808,7 +826,7 @@ class ListSetting(Setting):
         '''
         super().__init__(
             attr_name, nice_name, value,
-            tool_tip=tool_tip, show_ui=show_ui, allow_reset=allow_reset)
+            tool_tip=tool_tip, show_ui=show_ui, allow_reset=allow_reset, migrate_data=migrate_data)
         
         self.array_count_labels = {}
         self.element_layouts = {}
@@ -957,7 +975,8 @@ class StringListSetting(ListSetting):
         value: typing.List[str] = [],
         tool_tip: typing.Optional[str] = None,
         show_ui: bool = True,
-        allow_reset: bool = True
+        allow_reset: bool = True,
+        migrate_data: typing.Callable[[any], None] = None
     ):
         '''
         Create a new ArraySetting object.
@@ -971,7 +990,7 @@ class StringListSetting(ListSetting):
         '''
         super().__init__(
             attr_name, nice_name, value,
-            tool_tip=tool_tip, show_ui=show_ui, allow_reset=allow_reset)
+            tool_tip=tool_tip, show_ui=show_ui, allow_reset=allow_reset, migrate_data=migrate_data)
         pass
 
     def create_element(self, override_device_name: str, index: int) -> typing.Tuple[QtWidgets.QWidget, object]:
@@ -1362,7 +1381,8 @@ class LoggingSetting(Setting):
         verbosity_levels: typing.List[str] = None,
         tool_tip: typing.Optional[str] = None,
         show_ui: bool = True,
-        allow_reset: bool = True
+        allow_reset: bool = True,
+        migrate_data: typing.Callable[[any], None] = None
     ):
         '''
         Create a new LoggingSetting object.
@@ -1384,7 +1404,7 @@ class LoggingSetting(Setting):
 
         super().__init__(
             attr_name, nice_name, value,
-            tool_tip=tool_tip, show_ui=show_ui, allow_reset=allow_reset)
+            tool_tip=tool_tip, show_ui=show_ui, allow_reset=allow_reset, migrate_data=migrate_data)
 
         self._verbosity_levels = (
             verbosity_levels or self.DEFAULT_VERBOSITY_LEVELS)
diff --git a/Engine/Plugins/VirtualProduction/Switchboard/Source/Switchboard/switchboard/devices/ndisplay/plugin_ndisplay.py b/Engine/Plugins/VirtualProduction/Switchboard/Source/Switchboard/switchboard/devices/ndisplay/plugin_ndisplay.py
--- a/Engine/Plugins/VirtualProduction/Switchboard/Source/Switchboard/switchboard/devices/ndisplay/plugin_ndisplay.py
+++ b/Engine/Plugins/VirtualProduction/Switchboard/Source/Switchboard/switchboard/devices/ndisplay/plugin_ndisplay.py
@@ -3,6 +3,7 @@
 import concurrent.futures
 import json
 import os
+import typing
 from pathlib import Path
 import socket
 import struct
@@ -16,7 +17,8 @@ from switchboard import message_protocol
 from switchboard import switchboard_utils as sb_utils
 from switchboard import switchboard_widgets as sb_widgets
 from switchboard.config import CONFIG, BoolSetting, FilePathSetting, \
-    LoggingSetting, OptionSetting, Setting, StringSetting, SETTINGS, StringListSetting
+    LoggingSetting, OptionSetting, Setting, StringSetting, SETTINGS, StringListSetting, \
+    migrate_comma_separated_string_to_list
 from switchboard.devices.device_widget_base import AddDeviceDialog
 from switchboard.devices.unreal.plugin_unreal import DeviceUnreal, \
     DeviceWidgetUnreal
@@ -402,13 +404,15 @@ class DevicenDisplay(DeviceUnreal):
             nice_name='ExecCmds',
             value= [],
             tool_tip='ExecCmds to be passed. No need for outer double quotes.',
-            allow_reset=False
+            allow_reset=False,
+            migrate_data=migrate_comma_separated_string_to_list
         ),
         'ndisplay_dp_cvars': StringListSetting(
             attr_name='ndisplay_dp_cvars',
             nice_name="DPCVars",
             value=[],
-            tool_tip="Device profile console variables."
+            tool_tip="Device profile console variables.",
+            migrate_data=migrate_comma_separated_string_to_list
         ),
         'ndisplay_unattended': BoolSetting(
             attr_name='ndisplay_unattended',
diff --git a/Engine/Plugins/VirtualProduction/Switchboard/Source/Switchboard/switchboard/devices/unreal/plugin_unreal.py b/Engine/Plugins/VirtualProduction/Switchboard/Source/Switchboard/switchboard/devices/unreal/plugin_unreal.py
--- a/Engine/Plugins/VirtualProduction/Switchboard/Source/Switchboard/switchboard/devices/unreal/plugin_unreal.py
+++ b/Engine/Plugins/VirtualProduction/Switchboard/Source/Switchboard/switchboard/devices/unreal/plugin_unreal.py
@@ -23,7 +23,7 @@ from switchboard import switchboard_application
 from switchboard import switchboard_utils as sb_utils
 from switchboard.config import CONFIG, BoolSetting, DirectoryPathSetting, \
     IntSetting, MultiOptionSetting, OptionSetting, StringSetting, FileSystemPathSetting, \
-    SETTINGS, DEFAULT_MAP_TEXT, StringListSetting
+    SETTINGS, DEFAULT_MAP_TEXT, StringListSetting, migrate_comma_separated_string_to_list
 from switchboard.devices.device_base import Device, DeviceStatus, \
     PluginHeaderWidgets
 from switchboard.devices.device_widget_base import DeviceWidget, DeviceAutoJoinMUServerUI
@@ -284,12 +284,14 @@ class DeviceUnreal(Device):
             nice_name='ExecCmds',
             value=[],
             tool_tip='ExecCmds to be passed. No need for outer double quotes.',
+            migrate_data=migrate_comma_separated_string_to_list
         ),
         'dp_cvars': StringListSetting(
             attr_name='dp_cvars',
             nice_name="DPCVars",
             value=[],
-            tool_tip="Device profile console variables."
+            tool_tip="Device profile console variables.",
+            migrate_data=migrate_comma_separated_string_to_list
         ),
         'port': IntSetting(
             attr_name="port",
