Add combo boxes to array view

diff --git a/Engine/Plugins/VirtualProduction/Switchboard/Source/Switchboard/switchboard/config.py b/Engine/Plugins/VirtualProduction/Switchboard/Source/Switchboard/switchboard/config.py
--- a/Engine/Plugins/VirtualProduction/Switchboard/Source/Switchboard/switchboard/config.py
+++ b/Engine/Plugins/VirtualProduction/Switchboard/Source/Switchboard/switchboard/config.py
@@ -15,6 +15,7 @@ from PySide2 import QtWidgets
 
 from switchboard import switchboard_widgets as sb_widgets
 from switchboard.switchboard_logging import LOGGER
+from switchboard.switchboard_widgets import DropDownMenuComboBox
 
 ROOT_CONFIGS_PATH = pathlib.Path(__file__).parent.with_name('configs')
 CONFIG_SUFFIX = '.json'
@@ -772,22 +773,56 @@ class MultiOptionSetting(OptionSetting):
             widget.setEditText(widget.separator.join(selected_items))
 
 
-class ArrayRow(QtWidgets.QWidget):
-    def __init__(self, array_index: int, editor_widget: QtWidgets.QWidget, parent=None):
+class ListRow(QtWidgets.QWidget):
+    INSERT_TEXT = "Insert"
+    DUPLICATE_TEXT = "Duplicate"
+    DELETE_TEXT = "Delete"
+    
+    def __init__(
+        self,
+        array_index: int,
+        editor_widget: QtWidgets.QWidget,
+        insert_item_callback: typing.Callable[[], None] = None, 
+        duplicate_item_callback: typing.Callable[[], None] = None, 
+        delete_item_callback: typing.Callable[[], None] = None,
+        parent=None
+    ):
         super().__init__(parent=parent)
         self._editor_widget = editor_widget
+        self._insert_item_callback = insert_item_callback
+        self._duplicate_item_callback = duplicate_item_callback
+        self._delete_item_callback = delete_item_callback
     
         layout = QtWidgets.QHBoxLayout(self)
         layout.setSpacing(1)
         layout.setContentsMargins(0, 0, 0, 0)
         
+        # Shift the elements to the right
+        layout.addItem(
+            QtWidgets.QSpacerItem(10, 0, QtWidgets.QSizePolicy.Maximum, QtWidgets.QSizePolicy.Maximum)
+        )
+        
         self._index_label = QtWidgets.QLabel()
         self._index_label.setFixedWidth(20)
         self.update_index(array_index)
         layout.addWidget(self._index_label)
 
         layout.addWidget(editor_widget)
-        # TODO: Insert & delete combo box
+
+        element_actions = DropDownMenuComboBox()
+        element_actions.on_select_option.connect(self._on_view_option_selected)
+        element_actions.addItem(self.INSERT_TEXT)
+        element_actions.addItem(self.DUPLICATE_TEXT)
+        element_actions.addItem(self.DELETE_TEXT)
+        layout.addWidget(element_actions)
+
+    def _on_view_option_selected(self, selected_item):
+        if selected_item == self.INSERT_TEXT:
+            self._insert_item_callback()
+        elif selected_item == self.DUPLICATE_TEXT:
+            self._duplicate_item_callback()
+        elif selected_item == self.DELETE_TEXT:
+            self._delete_item_callback()
     
     @property
     def editor_widget(self):
@@ -796,6 +831,7 @@ class ArrayRow(QtWidgets.QWidget):
     def update_index(self, index: int):
         self._index_label.setText(str(index))
 
+
 class ListSetting(Setting):
     '''
     A setting which has add and clear buttons. Each item will be displayed in a new line.
@@ -858,7 +894,7 @@ class ListSetting(Setting):
         elements_layout = QtWidgets.QVBoxLayout(elements_root)
         self.element_layouts[override_device_name] = elements_layout
         elements_layout.setSpacing(1)
-        elements_layout.setContentsMargins(0, 0, 0, 0)
+        elements_layout.setContentsMargins(8, 5, 2, 2)
         root_layout.addWidget(
             elements_root
         )
@@ -916,12 +952,11 @@ class ListSetting(Setting):
         
     def _on_press_add(self, override_device_name: str):
         current_value = self.get_value(override_device_name)
-        
         row_widget, default_value = self._create_element_row(override_device_name, len(current_value))
         self.element_layouts[override_device_name].addWidget(
             row_widget
         )
-        
+
         self._on_widget_value_changed(current_value + [default_value], override_device_name)
         # Force update in case _on_widget_value_changed implicitly changed array settings without telling us
         self._on_setting_changed(self.get_value(override_device_name), override_device_name)
@@ -936,9 +971,55 @@ class ListSetting(Setting):
         # Force update in case _on_widget_value_changed implicitly changed array settings without telling us
         self._on_setting_changed(self.get_value(override_device_name), override_device_name)
     
-    def _create_element_row(self, override_device_name: str, index: int) -> typing.Tuple[ArrayRow, object]:
+    def _insert_at(self, override_device_name: str, index: int):
+        current_value = self.get_value(override_device_name)
+        row_widget, default_value = self._create_element_row(override_device_name, len(current_value))
+        self.element_layouts[override_device_name].addWidget(
+            row_widget
+        )
+
+        copied_value = current_value.copy()
+        copied_value.insert(index, default_value)
+        
+        self._on_widget_value_changed(copied_value, override_device_name)
+        # Force update in case _on_widget_value_changed implicitly changed array settings without telling us
+        self._on_setting_changed(self.get_value(override_device_name), override_device_name)
+        
+    def _duplicate_at(self, override_device_name: str, index: int):
+        current_value = self.get_value(override_device_name)
+        row_widget, _ = self._create_element_row(override_device_name, len(current_value))
+        self.element_layouts[override_device_name].addWidget(
+            row_widget
+        )
+
+        copied_value = current_value.copy()
+        copied_value.insert(index, current_value[index])
+
+        self._on_widget_value_changed(copied_value, override_device_name)
+        # Force update in case _on_widget_value_changed implicitly changed array settings without telling us
+        self._on_setting_changed(self.get_value(override_device_name), override_device_name)
+        
+    def _remove_at(self, override_device_name: str, index: int):
+        current_value = self.get_value(override_device_name)
+        copied_value = current_value.copy()
+        del copied_value[index]
+
+        self._on_widget_value_changed(copied_value, override_device_name)
+        # Force update in case _on_widget_value_changed implicitly changed array settings without telling us
+        self._on_setting_changed(self.get_value(override_device_name), override_device_name)
+
+    def _create_element_row(self, override_device_name: str, index: int) -> typing.Tuple[ListRow, object]:
         editing_widget, default_value = self.create_element(override_device_name, index)
-        row = ArrayRow(index, editing_widget)
+        row = ListRow(
+            index,
+            editing_widget, 
+            lambda override_device_name=override_device_name, index=index:
+                self._insert_at(override_device_name, index),
+            lambda override_device_name=override_device_name, index=index:
+                self._duplicate_at(override_device_name, index),
+            lambda override_device_name=override_device_name, index=index:
+                self._remove_at(override_device_name, index)
+        )
         return row, default_value
 
     def _on_setting_changed(self, new_value: typing.List, override_device_name: typing.Optional[str] = None):
@@ -958,7 +1039,7 @@ class ListSetting(Setting):
                 container_layout.itemAt(added_index).widget().deleteLater()
                 
         for index in range(new_len):
-            array_row: ArrayRow = container_layout.itemAt(index).widget()
+            array_row: ListRow = container_layout.itemAt(index).widget()
             self.update_element_value(array_row.editor_widget, new_value, index)
 
         self._update_array_count_label(override_device_name)
diff --git a/Engine/Plugins/VirtualProduction/Switchboard/Source/Switchboard/switchboard/switchboard_widgets.py b/Engine/Plugins/VirtualProduction/Switchboard/Source/Switchboard/switchboard/switchboard_widgets.py
--- a/Engine/Plugins/VirtualProduction/Switchboard/Source/Switchboard/switchboard/switchboard_widgets.py
+++ b/Engine/Plugins/VirtualProduction/Switchboard/Source/Switchboard/switchboard/switchboard_widgets.py
@@ -129,11 +129,17 @@ class MultiSelectionComboBox(QtWidgets.QComboBox):
 class DropDownMenuComboBox(QtWidgets.QComboBox):
     on_select_option = QtCore.Signal(str)
     
-    def __init__(self, icon: QtGui.QIcon, icon_size: int = 20, parent=None):
+    def __init__(self, icon: QtGui.QIcon = None, icon_size: int = 25, parent=None):
         super().__init__(parent=parent)
-        self.addItem(icon, "")
-        drop_down_arrow_size = 20
-        self.setFixedWidth(icon_size + drop_down_arrow_size)
+        
+        drop_down_arrow_size = 15
+        if icon is not None:
+            self.addItem(icon, "")
+            self.setFixedWidth(icon_size + drop_down_arrow_size)
+        else:
+            self.addItem("")
+            self.setFixedWidth(drop_down_arrow_size)
+        
         self.model().setData(self.model().index(0, 0), QtCore.QSize(100, 100), QtCore.Qt.SizeHintRole)
         self.view().setRowHidden(0, True)
         self.currentIndexChanged.connect(self._on_index_changed)
