Launch after modifying settings crashes

diff --git a/Engine/Plugins/VirtualProduction/Switchboard/Source/Switchboard/switchboard/config.py b/Engine/Plugins/VirtualProduction/Switchboard/Source/Switchboard/switchboard/config.py
--- a/Engine/Plugins/VirtualProduction/Switchboard/Source/Switchboard/switchboard/config.py
+++ b/Engine/Plugins/VirtualProduction/Switchboard/Source/Switchboard/switchboard/config.py
@@ -96,6 +96,7 @@ class Setting(QtCore.QObject):
         # want override highlighting.
         self._base_widget = None
         self._override_widgets = {}
+        self._on_setting_changed_lambdas = {}
         
         # Appears when override value is different from _value
         self._allow_reset = allow_reset
@@ -255,11 +256,7 @@ class Setting(QtCore.QObject):
 
         # Give the Setting's UI an opportunity to update itself when the
         # underlying value changes.
-        self.signal_setting_changed.connect(
-            lambda old_value, new_value,
-            override_device_name=override_device_name:
-                self._on_setting_changed(
-                    new_value, override_device_name=override_device_name))
+        self._register_on_setting_changed(top_level_widget, override_device_name)
 
         if top_level_widget and form_layout:
             setting_label = QtWidgets.QLabel()
@@ -274,6 +271,23 @@ class Setting(QtCore.QObject):
 
         return top_level_widget
     
+    def _register_on_setting_changed(self, top_level_widget: QtWidgets.QWidget, override_device_name: str):
+        on_setting_changed_lambda = lambda old_value, new_value, override_device_name=override_device_name: \
+            self._on_setting_changed(new_value, override_device_name=override_device_name)
+        self.signal_setting_changed.connect(
+            on_setting_changed_lambda
+        )
+        
+        # Clear the widget when it is destroyed to avoid dangling references
+        top_level_widget.destroyed.connect(
+            lambda destroyed_object=None, on_setting_changed_lambda=on_setting_changed_lambda, override_device_name=override_device_name:
+                self._on_widget_destroyed(on_setting_changed_lambda, override_device_name)
+        )
+        
+    def _on_widget_destroyed(self, on_setting_changed_lambda, override_device_name: str):
+        self.signal_setting_changed.disconnect(on_setting_changed_lambda)
+        self.set_widget(widget=None, override_device_name=override_device_name)
+    
     def _decorate_with_reset_widget(self, override_device_name: str, setting_editor_widget: QtWidgets.QWidget):
         # Reset will still be shown on overrides
         if not self._allow_reset and override_device_name is None:
