Disambiguate map paths with full path

diff --git a/Engine/Plugins/VirtualProduction/Switchboard/Source/Switchboard/switchboard/config.py b/Engine/Plugins/VirtualProduction/Switchboard/Source/Switchboard/switchboard/config.py
--- a/Engine/Plugins/VirtualProduction/Switchboard/Source/Switchboard/switchboard/config.py
+++ b/Engine/Plugins/VirtualProduction/Switchboard/Source/Switchboard/switchboard/config.py
@@ -2222,22 +2222,32 @@ class Config(object):
         self.save()
 
     def maps(self):
+        '''
+        Returns a list of ful map paths in an Unreal Engine project such as [ "/Game/Maps/MapName" ].
+        It will always start with Game and the slashes will always be "/" independent of the platform's separator.
+        '''
+        project_dir = os.path.dirname(self.UPROJECT_PATH.get_value().replace('"', ''))
         maps_path = os.path.normpath(
             os.path.join(
-                os.path.dirname(
-                    self.UPROJECT_PATH.get_value().replace('"', '')),
+                project_dir,
                 'Content',
                 self.MAPS_PATH.get_value()))
 
         maps = []
-        for _, _, files in os.walk(maps_path):
+        for path_to_map, b, files in os.walk(maps_path):
             for name in files:
                 if not fnmatch.fnmatch(name, self.MAPS_FILTER.get_value()):
                     continue
-
-                rootname, _ = os.path.splitext(name)
-                if rootname not in maps:
-                    maps.append(rootname)
+                
+                map_name, _ = os.path.splitext(name)
+                path_name = path_to_map.replace(project_dir, '', 1)
+                # Ignore the fact that maps may exist in plugins - we only search game content
+                path_name = path_name.replace('Content', 'Game')
+                path_name = os.path.join(path_name, map_name)
+                path_name = path_name.replace(os.sep, '/')
+                
+                if path_name not in maps:
+                    maps.append(path_name)
 
         maps.sort()
         return maps
diff --git a/Engine/Plugins/VirtualProduction/Switchboard/Source/Switchboard/switchboard/switchboard_dialog.py b/Engine/Plugins/VirtualProduction/Switchboard/Source/Switchboard/switchboard/switchboard_dialog.py
--- a/Engine/Plugins/VirtualProduction/Switchboard/Source/Switchboard/switchboard/switchboard_dialog.py
+++ b/Engine/Plugins/VirtualProduction/Switchboard/Source/Switchboard/switchboard/switchboard_dialog.py
@@ -316,7 +316,7 @@ class SwitchboardDialog(QtCore.QObject):
         self.window.slate_line_edit.textChanged.connect(self._set_slate)
         self.window.take_spin_box.valueChanged.connect(self._set_take)
         self.window.sequence_line_edit.textChanged.connect(self._set_sequence)
-        self.window.level_combo_box.currentTextChanged.connect(self._set_level)
+        self.window.level_combo_box.currentIndexChanged.connect(self._on_selected_level_changed)
         self.window.refresh_levels_button.clicked.connect(self.refresh_levels_incremental)
         self.window.project_cl_combo_box.currentTextChanged.connect(self._set_project_changelist)
         self.window.engine_cl_combo_box.currentTextChanged.connect(self._set_engine_changelist)
@@ -1281,6 +1281,14 @@ class SwitchboardDialog(QtCore.QObject):
     def level(self, value):
         self._set_level(value)
 
+    def _on_selected_level_changed(self, index: int):
+        full_map_path = self._get_level_from_combo_box(index)
+        self._set_level(full_map_path)
+        
+    def _get_level_from_combo_box(self, index: int):
+        # Tooltip stores full path
+        return self.window.level_combo_box.itemData(index, QtCore.Qt.ToolTipRole)
+
     def _set_level(self, value):
         ''' Called when level dropdown text changes
         '''
@@ -1675,14 +1683,53 @@ class SwitchboardDialog(QtCore.QObject):
         current_level = CONFIG.CURRENT_LEVEL
 
         self.window.level_combo_box.clear()
-        self.window.level_combo_box.addItems([DEFAULT_MAP_TEXT] + levels)
+        self._update_level_list(self.window.level_combo_box, levels)
 
         if current_level and current_level in levels:
             self.level = current_level
 
+
+    def _update_level_list(self, level_combo_box: QtWidgets.QComboBox, level_path_list: List[str]):
+        def compare_file_names(path_a: str, path_b: str):
+            return -1 if path_a.lower() < path_b.lower() \
+                else 1 if path_a.lower() > path_b.lower() else 0
+        
+        def generate_short_map_path(map_path: str, file_name: str) -> str:
+            map_path = map_path.replace("/Game", "", 1)
+            map_path = map_path.removesuffix(file_name)
+            return f"{file_name} ({map_path})"
+            
+        
+        name_counts = {}
+        for map_path in level_path_list:
+            file_name = os.path.basename(map_path)
+            name_counts[file_name] = name_counts.get(file_name, 0) + 1
+            
+        # Show only level name if unique and show path behind to disambiguate duplicates
+        short_name_list = []
+        short_name_to_path = {}
+        for map_path in level_path_list:
+            file_name = os.path.basename(map_path)
+            short_name = file_name if name_counts[file_name] == 1 else generate_short_map_path(map_path, file_name)
+            short_name_list.append(short_name)
+            short_name_to_path[short_name] = map_path
+            
+        from functools import cmp_to_key
+        short_name_list = sorted(short_name_list, key=cmp_to_key(compare_file_names))
+        level_combo_box.addItems([DEFAULT_MAP_TEXT] + short_name_list)
+        
+        if len(level_path_list) == 0:
+            return
+        
+        # To disambiguate, show the full path name in the drop-down
+        level_combo_box.setItemData(0, "Default level", QtCore.Qt.ToolTipRole)
+        for short_name_index in range(len(short_name_list)):
+            short_name = short_name_list[short_name_index]
+            level_combo_box.setItemData(short_name_index + 1, short_name_to_path[short_name], QtCore.Qt.ToolTipRole)
+
     def get_current_level_list(self):
         level_combo = self.window.level_combo_box
-        return [level_combo.itemText(i) for i in range(1, level_combo.count())] # skip DEFAULT_MAP_TEXT
+        return [self._get_level_from_combo_box(i) for i in range(1, level_combo.count())] # skip DEFAULT_MAP_TEXT
 
     def refresh_levels_incremental(self):
         ''' Wrapper around `refresh_levels` with the following differences:
