Refactor general settings

diff --git a/Engine/Plugins/VirtualProduction/Switchboard/Source/Switchboard/switchboard/config.py b/Engine/Plugins/VirtualProduction/Switchboard/Source/Switchboard/switchboard/config.py
--- a/Engine/Plugins/VirtualProduction/Switchboard/Source/Switchboard/switchboard/config.py
+++ b/Engine/Plugins/VirtualProduction/Switchboard/Source/Switchboard/switchboard/config.py
@@ -464,7 +464,7 @@ class IntSetting(Setting):
 
         old_value = int(widget.text())
         if new_value != old_value:
-            widget.setText(new_value)
+            widget.setText(str(new_value))
             widget.setCursorPosition(0)
 
 
@@ -1790,142 +1790,6 @@ class Config(object):
     def __init__(self, file_path: typing.Union[str, pathlib.Path]):
         self.init_with_file_path(file_path)
 
-    def init_new_config(self, file_path: typing.Union[str, pathlib.Path],
-                        uproject, engine_dir, p4_settings):
-        ''' Initialize new configuration
-        '''
-
-        self.file_path = get_absolute_config_path(file_path)
-        self.PROJECT_NAME = self.file_path.stem
-        self.UPROJECT_PATH = StringSetting(
-            "uproject", "uProject Path", uproject, tool_tip="Path to uProject")
-        self.SWITCHBOARD_DIR = os.path.abspath(
-            os.path.join(os.path.dirname(os.path.abspath(__file__)), '../'))
-        self.ENGINE_DIR = StringSetting(
-            "engine_dir", "Engine Directory", engine_dir,
-            tool_tip="Path to UE4 engine directory")
-        self.BUILD_ENGINE = BoolSetting(
-            "build_engine", "Build Engine", False,
-            tool_tip="Is Engine built from source?")
-        self.MAPS_PATH = StringSetting(
-            "maps_path", "Map Path", "",
-            tool_tip=(
-                "Relative path from Content folder that contains maps to "
-                "launch into."))
-        self.MAPS_FILTER = StringSetting(
-            "maps_filter", "Map Filter", "*.umap",
-            tool_tip=(
-                "Walk every file in the Map Path and run a fnmatch to filter "
-                "the file names"))
-        self.P4_ENABLED = BoolSetting(
-            "p4_enabled", "Perforce Enabled", p4_settings['p4_enabled'],
-            tool_tip="Toggle Perforce support for the entire application")
-        self.SOURCE_CONTROL_WORKSPACE = StringSetting(
-            "source_control_workspace", "Workspace Name",
-            p4_settings['p4_workspace_name'],
-            tool_tip="SourceControl Workspace/Branch")
-
-        self.P4_PROJECT_PATH = PerforcePathSetting(
-            attr_name="p4_sync_path",
-            nice_name="Perforce Project Path",
-            value=p4_settings['p4_project_path']
-        )
-
-        self.P4_ENGINE_PATH = PerforcePathSetting(
-            attr_name="p4_engine_path",
-            nice_name="Perforce Engine Path",
-            value=p4_settings['p4_engine_path']
-        )
-
-        self.init_unreal_insights()
-
-        self.CURRENT_LEVEL = DEFAULT_MAP_TEXT
-
-        self.OSC_SERVER_PORT = IntSetting(
-            "osc_server_port", "OSC Server Port", 6000)
-        self.OSC_CLIENT_PORT = IntSetting(
-            "osc_client_port", "OSC Client Port", 8000)
-
-        # MU Settings
-        self.MULTIUSER_SERVER_EXE = 'UnrealMultiUserServer'
-
-        self.init_muserver()
-
-        self.LISTENER_EXE = 'SwitchboardListener'
-
-        self._device_data_from_config = {}
-        self._plugin_data_from_config = {}
-        self._plugin_settings = {}
-        self._device_settings = {}
-
-        LOGGER.info(f"Creating new config saved in {self.file_path}")
-        self.save()
-
-        SETTINGS.CONFIG = self.file_path
-        SETTINGS.save()
-
-    def init_muserver(self, data={}):
-        # MU Settings
-        self.MUSERVER_COMMAND_LINE_ARGUMENTS = data.get(
-            'muserver_command_line_arguments', '')
-        self.MUSERVER_SERVER_NAME = data.get(
-            'muserver_server_name', f'{self.PROJECT_NAME}_MU_Server')
-        self.MUSERVER_ENDPOINT = data.get('muserver_endpoint', ':9030')
-        self.MUSERVER_MULTICAST_ENDPOINT = StringSetting(
-            attr_name='udpmessaging_multicast_endpoint',
-            nice_name='Multicast Endpoint',
-            value=data.get('muserver_multicast_endpoint', '230.0.0.1:6666'),
-            tool_tip=(
-                'Multicast group and port (-UDPMESSAGING_TRANSPORT_MULTICAST) '
-                'in the {ip}:{port} endpoint format. The multicast group IP '
-                'must be in the range 224.0.0.0 to 239.255.255.255.'),
-        )
-        self.MUSERVER_AUTO_LAUNCH = data.get('muserver_auto_launch', True)
-        self.MUSERVER_CLEAN_HISTORY = data.get('muserver_clean_history', False)
-        self.MUSERVER_AUTO_BUILD = data.get('muserver_auto_build', True)
-        self.MUSERVER_AUTO_ENDPOINT = data.get('muserver_auto_endpoint', True)
-
-        self.MUSERVER_AUTO_JOIN = BoolSetting(
-            "muserver_auto_join",
-            "Unreal Multi-user Server Auto-join",
-            data.get('muserver_auto_join', True)
-        )
-
-    def init_unreal_insights(self, data={}):
-        self.INSIGHTS_TRACE_ENABLE = BoolSetting(
-            "tracing_enabled",
-            "Unreal Insights Tracing State",
-            data.get("tracing_enabled", False),
-        )
-        self.INSIGHTS_TRACE_ARGS = StringSetting(
-            "tracing_args",
-            "Unreal Insights Tracing Args",
-            data.get('tracing_arg', 'log,cpu,gpu,frame,bookmark,concert,messaging')
-        )
-        self.INSIGHTS_STAT_EVENTS = BoolSetting(
-            "tracing_stat_events",
-            "Unreal Insights Tracing with Stat Events",
-            data.get('tracing_stat_events', True)
-        )
-
-    def save_muserver(self, data):
-        data["muserver_command_line_arguments"] = (
-            self.MUSERVER_COMMAND_LINE_ARGUMENTS)
-        data["muserver_server_name"] = self.MUSERVER_SERVER_NAME
-        data["muserver_endpoint"] = self.MUSERVER_ENDPOINT
-        data["muserver_auto_launch"] = self.MUSERVER_AUTO_LAUNCH
-        data["muserver_clean_history"] = self.MUSERVER_CLEAN_HISTORY
-        data["muserver_auto_build"] = self.MUSERVER_AUTO_BUILD
-        data["muserver_auto_endpoint"] = self.MUSERVER_AUTO_ENDPOINT
-
-        data["muserver_auto_join"] = self.MUSERVER_AUTO_JOIN.get_value()
-        data["muserver_multicast_endpoint"] = self.MUSERVER_MULTICAST_ENDPOINT.get_value()
-
-    def save_unreal_insights(self, data):
-        data['tracing_enabled'] = self.INSIGHTS_TRACE_ENABLE.get_value()
-        data['tracing_args'] = self.INSIGHTS_TRACE_ARGS.get_value()
-        data['tracing_stat_events'] = self.INSIGHTS_STAT_EVENTS.get_value()
-
     def init_with_file_path(self, file_path: typing.Union[str, pathlib.Path]):
         if file_path:
             try:
@@ -1943,103 +1807,32 @@ class Config(object):
             self.file_path = None
             data = {}
 
-        project_settings = []
-
-        self.PROJECT_NAME = data.get('project_name', 'Default')
-        self.UPROJECT_PATH = StringSetting(
-            "uproject", "uProject Path", data.get('uproject', ''),
-            tool_tip="Path to uProject")
-        project_settings.append(self.UPROJECT_PATH)
-
-        # Directory Paths
-        self.SWITCHBOARD_DIR = os.path.abspath(
-            os.path.join(os.path.dirname(os.path.abspath(__file__)), '../'))
-        self.ENGINE_DIR = StringSetting(
-            "engine_dir", "Engine Directory", data.get('engine_dir', ''),
-            tool_tip="Path to UE4 engine directory")
-        project_settings.append(self.ENGINE_DIR)
-        self.BUILD_ENGINE = BoolSetting(
-            "build_engine", "Build Engine", data.get('build_engine', False),
-            tool_tip="Is Engine built from source?")
-        project_settings.append(self.BUILD_ENGINE)
-        self.MAPS_PATH = StringSetting(
-            "maps_path", "Map Path", data.get('maps_path', ''),
-            placeholder_text="Maps",
-            tool_tip=(
-                "Relative path from Content folder that contains maps to "
-                "launch into."))
-        project_settings.append(self.MAPS_PATH)
-        self.MAPS_FILTER = StringSetting(
-            "maps_filter", "Map Filter", data.get('maps_filter', '*.umap'),
-            placeholder_text="*.umap",
-            tool_tip=(
-                "Walk every file in the Map Path and run a fnmatch to "
-                "filter the file names"))
-        project_settings.append(self.MAPS_FILTER)
-
-        # OSC settings
-        self.OSC_SERVER_PORT = IntSetting(
-            "osc_server_port", "OSC Server Port",
-            data.get('osc_server_port', 6000))
-        self.OSC_CLIENT_PORT = IntSetting(
-            "osc_client_port", "OSC Client Port",
-            data.get('osc_client_port', 8000))
-        project_settings.extend([self.OSC_SERVER_PORT, self.OSC_CLIENT_PORT])
-
-        # Perforce settings
-        self.P4_ENABLED = BoolSetting(
-            "p4_enabled", "Perforce Enabled", data.get("p4_enabled", False),
-            tool_tip="Toggle Perforce support for the entire application")
-        self.SOURCE_CONTROL_WORKSPACE = StringSetting(
-            "source_control_workspace", "Workspace Name",
-            data.get("source_control_workspace"),
-            tool_tip="SourceControl Workspace/Branch")
-
-        self.P4_PROJECT_PATH = PerforcePathSetting(
-            "p4_sync_path",
-            "Perforce Project Path",
-            data.get("p4_sync_path", ''),
-            placeholder_text="//UE4/Project"
-        )
-
-        self.P4_ENGINE_PATH = PerforcePathSetting(
-            "p4_engine_path",
-            "Perforce Engine Path",
-            data.get("p4_engine_path", ''),
-            placeholder_text="//UE4/Project/Engine"
-        )
-
-        project_settings.extend(
-            [self.P4_ENABLED, self.SOURCE_CONTROL_WORKSPACE,
-             self.P4_PROJECT_PATH, self.P4_ENGINE_PATH])
-
+        self.init_switchboard_settings()
+        self.init_project_settings(data)
         self.init_unreal_insights(data)
-
-        project_settings.extend([
-            self.INSIGHTS_TRACE_ENABLE,
-            self.INSIGHTS_STAT_EVENTS,
-            self.INSIGHTS_TRACE_ARGS])
-
-        # EXE names
-        self.MULTIUSER_SERVER_EXE = data.get(
-            'multiuser_exe', 'UnrealMultiUserServer')
-        self.LISTENER_EXE = data.get('listener_exe', 'SwitchboardListener')
-
         self.init_muserver(data)
 
-        # MISC SETTINGS
-        self.CURRENT_LEVEL = data.get('current_level', DEFAULT_MAP_TEXT)
-
-        project_settings.extend([
-            self.MUSERVER_AUTO_JOIN])
-
         # Automatically save whenever a project setting is changed or
         # overridden by a device.
-        for setting in project_settings:
+        # TODO: switchboard_settings
+        all_settings = [setting for _, setting in self.basic_project_settings.items()] \
+            + [setting for _, setting in self.osc_settings.items()] \
+            + [setting for _, setting in self.source_control_settings.items()] \
+            + [setting for _, setting in self.unreal_insight_settings.items()] \
+            + [setting for _, setting in self.mu_settings.items()]
+            # TODO: multiuser
+        for setting in all_settings:
             setting.signal_setting_changed.connect(lambda: self.save())
             setting.signal_setting_overridden.connect(
                 self.on_device_override_changed)
 
+        # Directory Paths
+        self.SWITCHBOARD_DIR = os.path.abspath(
+            os.path.join(os.path.dirname(os.path.abspath(__file__)), '../'))
+
+        # MISC SETTINGS
+        self.CURRENT_LEVEL = data.get('current_level', DEFAULT_MAP_TEXT)
+
         # Devices
         self._device_data_from_config = {}
         self._plugin_data_from_config = {}
@@ -2063,6 +1856,244 @@ class Config(object):
                     self._device_data_from_config.setdefault(
                         device_type, []).append(device_data)
 
+    def init_new_config(self, file_path: typing.Union[str, pathlib.Path], uproject, engine_dir, p4_settings):
+        ''' 
+        Initialize new configuration
+        '''
+
+        self.file_path = get_absolute_config_path(file_path)
+        self.init_switchboard_settings()
+        self.init_project_settings({ "project_name": self.file_path.stem })
+        self.init_unreal_insights()
+        self.init_muserver()
+
+        self.CURRENT_LEVEL = DEFAULT_MAP_TEXT
+
+        self._device_data_from_config = {}
+        self._plugin_data_from_config = {}
+        self._plugin_settings = {}
+        self._device_settings = {}
+
+        LOGGER.info(f"Creating new config saved in {self.file_path}")
+        self.save()
+
+        SETTINGS.CONFIG = self.file_path
+        SETTINGS.save()
+        
+    def init_switchboard_settings(self, data={}):
+        self.switchboard_settings = {
+            "listener_exe": StringSetting(
+                "listener_exe",
+                "Listener Executable Name",
+                 data.get('listener_exe', 'SwitchboardListener')
+            )
+        }
+        
+        self.LISTENER_EXE = self.switchboard_settings["listener_exe"]
+
+    def init_project_settings(self, data={}):
+        self.basic_project_settings = {
+            "project_name": StringSetting(
+                "project_name",
+                "Project Name",
+                data.get('project_name', 'Default')
+            ),
+            "uproject": FilePathSetting(
+                "uproject", "uProject Path",
+                data.get('uproject', ''),
+                tool_tip="Path to uProject"
+            ),
+            "engine_dir": FilePathSetting(
+                "engine_dir",
+                "Engine Directory",
+                data.get('engine_dir', ''),
+                tool_tip="Path to UE4 engine directory"
+            ),
+            "build_engine": BoolSetting(
+                "build_engine",
+                "Build Engine",
+                data.get('build_engine', False),
+                tool_tip="Is Engine built from source?"
+            ),
+            "maps_path": StringSetting(
+                "maps_path",
+                "Map Path",
+                data.get('maps_path', ''),
+                placeholder_text="Maps",
+                tool_tip="Relative path from Content folder that contains maps to launch into."
+            ),
+            "maps_filter": StringSetting(
+                "maps_filter",
+                "Map Filter",
+                data.get('maps_filter', '*.umap'),
+                placeholder_text="*.umap",
+                tool_tip="Walk every file in the Map Path and run a fnmatch to filter the file names"
+            )
+        }
+
+        self.PROJECT_NAME = self.basic_project_settings["project_name"]
+        self.UPROJECT_PATH = self.basic_project_settings["uproject"]
+        self.ENGINE_DIR = self.basic_project_settings["engine_dir"]
+        self.BUILD_ENGINE = self.basic_project_settings["build_engine"]
+        self.MAPS_PATH = self.basic_project_settings["maps_path"]
+        self.MAPS_FILTER = self.basic_project_settings["maps_filter"]
+
+        self.osc_settings = {
+            "osc_server_port": IntSetting(
+                "osc_server_port",
+                "OSC Server Port",
+                data.get('osc_server_port', 6000)
+            ),
+            "osc_client_port": IntSetting(
+                "osc_client_port",
+                "OSC Client Port",
+                data.get('osc_client_port', 8000)
+            )
+        }
+
+        self.OSC_SERVER_PORT = self.osc_settings["osc_server_port"]
+        self.OSC_CLIENT_PORT = self.osc_settings["osc_client_port"]
+
+        self.source_control_settings = {
+            "p4_enabled": BoolSetting(
+                "p4_enabled",
+                "Perforce Enabled",
+                data.get("p4_enabled", False),
+                tool_tip="Toggle Perforce support for the entire application"
+            ),
+            "source_control_workspace": StringSetting(
+                "source_control_workspace", "Workspace Name",
+                data.get("source_control_workspace"),
+                tool_tip="SourceControl Workspace/Branch"
+            ),
+            "p4_sync_path": PerforcePathSetting(
+                "p4_sync_path",
+                "Perforce Project Path",
+                data.get("p4_sync_path", ''),
+                placeholder_text="//UE4/Project"
+            ),
+            "p4_engine_path": PerforcePathSetting(
+                "p4_engine_path",
+                "Perforce Engine Path",
+                data.get("p4_engine_path", ''),
+                placeholder_text="//UE4/Project/Engine"
+            )
+        }
+
+        self.P4_ENABLED = self.source_control_settings["p4_enabled"]
+        self.SOURCE_CONTROL_WORKSPACE = self.source_control_settings["source_control_workspace"]
+        self.P4_PROJECT_PATH = self.source_control_settings["p4_sync_path"]
+        self.P4_ENGINE_PATH = self.source_control_settings["p4_engine_path"]
+
+    def init_unreal_insights(self, data={}):
+        self.unreal_insight_settings = {
+            "tracing_enabled": BoolSetting(
+                "tracing_enabled",
+                "Unreal Insights Tracing State",
+                data.get("tracing_enabled", False),
+            ),
+            "tracing_args": StringSetting(
+                "tracing_args",
+                "Unreal Insights Tracing Args",
+                data.get('tracing_arg', 'log,cpu,gpu,frame,bookmark,concert,messaging')
+            ),
+            "tracing_stat_events": BoolSetting(
+                "tracing_stat_events",
+                "Unreal Insights Tracing with Stat Events",
+                data.get('tracing_stat_events', True)
+            )
+        }
+
+        self.INSIGHTS_TRACE_ENABLE = self.unreal_insight_settings["tracing_enabled"]
+        self.INSIGHTS_TRACE_ARGS = self.unreal_insight_settings["tracing_args"]
+        self.INSIGHTS_STAT_EVENTS = self.unreal_insight_settings["tracing_stat_events"]
+
+    def init_muserver(self, data={}):
+        self.mu_settings = {
+            "muserver_server_name": StringSetting(
+                "muserver_server_name",
+                "Command Line Args",
+                data.get('muserver_server_name', f'{self.PROJECT_NAME.get_value()}_MU_Server')
+            ),
+            "muserver_command_line_arguments": StringSetting(
+                "muserver_command_line_arguments",
+                "Server name",
+                data.get('muserver_command_line_arguments', '')
+            ),
+            "muserver_endpoint": StringSetting(
+                "muserver_endpoint",
+                "Unicast Endpoint",
+                data.get('muserver_endpoint', ':9030')
+            ),
+            "udpmessaging_multicast_endpoint": StringSetting(
+                attr_name='udpmessaging_multicast_endpoint',
+                nice_name='Multicast Endpoint',
+                value=data.get('muserver_multicast_endpoint', '230.0.0.1:6666'),
+                tool_tip=(
+                    'Multicast group and port (-UDPMESSAGING_TRANSPORT_MULTICAST) '
+                    'in the {ip}:{port} endpoint format. The multicast group IP '
+                    'must be in the range 224.0.0.0 to 239.255.255.255.'),
+            ),
+            "multiuser_exe": StringSetting(
+                "multiuser_exe",
+                "Multiuser Executable Name",
+                data.get('multiuser_exe', 'UnrealMultiUserServer')
+            ),
+            "muserver_auto_launch": BoolSetting(
+                "muserver_auto_launch",
+                "Auto Launch",
+                data.get('muserver_auto_launch', True)
+            ),
+            "muserver_clean_history": BoolSetting(
+                "muserver_clean_history",
+                "Clean History",
+                data.get('muserver_clean_history', False)
+            ),
+            "muserver_auto_build": BoolSetting(
+                "muserver_auto_build",
+                "Auto Build",
+                data.get('muserver_auto_build', True)
+            ),
+            "muserver_auto_endpoint": BoolSetting(
+                "muserver_auto_endpoint",
+                "Auto Endpoint",
+                data.get('muserver_auto_endpoint', True)
+            ),
+            "muserver_auto_join": BoolSetting(
+                "muserver_auto_join",
+                "Unreal Multi-user Server Auto-join",
+                data.get('muserver_auto_join', True)
+            )
+        }
+        
+        self.MUSERVER_SERVER_NAME = self.mu_settings["muserver_server_name"]
+        self.MUSERVER_COMMAND_LINE_ARGUMENTS = self.mu_settings["muserver_command_line_arguments"]
+        self.MUSERVER_ENDPOINT = self.mu_settings["muserver_endpoint"]
+        self.MUSERVER_MULTICAST_ENDPOINT = self.mu_settings["udpmessaging_multicast_endpoint"]
+        self.MULTIUSER_SERVER_EXE = self.mu_settings["multiuser_exe"]
+        self.MUSERVER_AUTO_LAUNCH = self.mu_settings["muserver_auto_launch"]
+        self.MUSERVER_CLEAN_HISTORY = self.mu_settings["muserver_clean_history"]
+        self.MUSERVER_AUTO_BUILD = self.mu_settings["muserver_auto_build"]
+        self.MUSERVER_AUTO_ENDPOINT = self.mu_settings["muserver_auto_endpoint"]
+        self.MUSERVER_AUTO_JOIN = self.mu_settings["muserver_auto_join"]
+
+    def save_unreal_insights(self, data):
+        data['tracing_enabled'] = self.INSIGHTS_TRACE_ENABLE.get_value()
+        data['tracing_args'] = self.INSIGHTS_TRACE_ARGS.get_value()
+        data['tracing_stat_events'] = self.INSIGHTS_STAT_EVENTS.get_value()
+        
+    def save_muserver(self, data):
+        data["muserver_command_line_arguments"] = self.MUSERVER_COMMAND_LINE_ARGUMENTS.get_value()
+        data["muserver_server_name"] = self.MUSERVER_SERVER_NAME.get_value()
+        data["muserver_endpoint"] = self.MUSERVER_ENDPOINT.get_value()
+        data["multiuser_exe"] = self.MULTIUSER_SERVER_EXE.get_value()
+        data["muserver_auto_launch"] = self.MUSERVER_AUTO_LAUNCH.get_value()
+        data["muserver_clean_history"] = self.MUSERVER_CLEAN_HISTORY.get_value()
+        data["muserver_auto_build"] = self.MUSERVER_AUTO_BUILD.get_value()
+        data["muserver_auto_endpoint"] = self.MUSERVER_AUTO_ENDPOINT.get_value()
+        data["muserver_multicast_endpoint"] = self.MUSERVER_MULTICAST_ENDPOINT.get_value()
+        data["muserver_auto_join"] = self.MUSERVER_AUTO_JOIN.get_value()
+
     def load_plugin_settings(self, device_type, settings):
         ''' Updates plugin settings values with those read from the config file.
         '''
@@ -2123,13 +2154,13 @@ class Config(object):
         data = {}
 
         # General settings
-        data['project_name'] = self.PROJECT_NAME
+        data['project_name'] = self.PROJECT_NAME.get_value()
         data['uproject'] = self.UPROJECT_PATH.get_value()
         data['engine_dir'] = self.ENGINE_DIR.get_value()
         data['build_engine'] = self.BUILD_ENGINE.get_value()
         data["maps_path"] = self.MAPS_PATH.get_value()
         data["maps_filter"] = self.MAPS_FILTER.get_value()
-        data["listener_exe"] = self.LISTENER_EXE
+        data["listener_exe"] = self.LISTENER_EXE.get_value()
 
         self.save_unreal_insights(data)
 
@@ -2144,9 +2175,6 @@ class Config(object):
         data["source_control_workspace"] = (
             self.SOURCE_CONTROL_WORKSPACE.get_value())
 
-        # MU Settings
-        data["multiuser_exe"] = self.MULTIUSER_SERVER_EXE
-
         self.save_muserver(data)
 
         # Current Level
@@ -2254,11 +2282,11 @@ class Config(object):
 
     def multiuser_server_path(self):
         return self.engine_exe_path(
-            self.ENGINE_DIR.get_value(), self.MULTIUSER_SERVER_EXE)
+            self.ENGINE_DIR.get_value(), self.MULTIUSER_SERVER_EXE.get_value())
 
     def listener_path(self):
         return self.engine_exe_path(
-            self.ENGINE_DIR.get_value(), self.LISTENER_EXE)
+            self.ENGINE_DIR.get_value(), self.LISTENER_EXE.get_value())
 
     # todo-dara: find a way to do this directly in the LiveLinkFace plugin code
     def unreal_device_ip_addresses(self):
@@ -2334,10 +2362,17 @@ class UserSettings(object):
             self.CONFIG = config_paths[0] if config_paths else None
 
         # IP Address of the machine running Switchboard
-        self.IP_ADDRESS = data.get(
-            'ip_address', socket.gethostbyname(socket.gethostname()))
-
-        self.TRANSPORT_PATH = data.get('transport_path', '')
+        self.IP_ADDRESS = StringSetting(
+            "ip_address",
+            "IP Address",
+            data.get("ip_address", socket.gethostbyname(socket.gethostname()))
+        )
+        self.TRANSPORT_PATH = FilePathSetting(
+                "transport_path",
+                "Transport path",
+                data.get('transport_path', '')
+            )
+        
 
         # UI Settings
         self.MUSERVER_SESSION_NAME = data.get(
@@ -2354,8 +2389,8 @@ class UserSettings(object):
     def save(self):
         data = {
             'config': '',
-            'ip_address': self.IP_ADDRESS,
-            'transport_path': self.TRANSPORT_PATH,
+            'ip_address': self.IP_ADDRESS.get_value(),
+            'transport_path': self.TRANSPORT_PATH.get_value(),
             'muserver_session_name': self.MUSERVER_SESSION_NAME,
             'current_sequence': self.CURRENT_SEQUENCE,
             'current_slate': self.CURRENT_SLATE,
diff --git a/Engine/Plugins/VirtualProduction/Switchboard/Source/Switchboard/switchboard/devices/livelinkface/plugin_livelinkface.py b/Engine/Plugins/VirtualProduction/Switchboard/Source/Switchboard/switchboard/devices/livelinkface/plugin_livelinkface.py
--- a/Engine/Plugins/VirtualProduction/Switchboard/Source/Switchboard/switchboard/devices/livelinkface/plugin_livelinkface.py
+++ b/Engine/Plugins/VirtualProduction/Switchboard/Source/Switchboard/switchboard/devices/livelinkface/plugin_livelinkface.py
@@ -107,7 +107,7 @@ class DeviceLiveLinkFace(Device):
 
     # Tell the mobile which ip_address is running Switchboard
     def osc_add_send_target(self):
-        self.send_osc_message(osc.OSC_ADD_SEND_TARGET, [SETTINGS.IP_ADDRESS, CONFIG.OSC_SERVER_PORT.get_value()], log=True)
+        self.send_osc_message(osc.OSC_ADD_SEND_TARGET, [SETTINGS.IP_ADDRESS.get_value(), CONFIG.OSC_SERVER_PORT.get_value()], log=True)
 
     def osc_add_send_target_confirm(self):
         # If the mobile is left on and CP is relaunched the normal bootup sequence doesn't happen 
diff --git a/Engine/Plugins/VirtualProduction/Switchboard/Source/Switchboard/switchboard/devices/ndisplay/plugin_ndisplay.py b/Engine/Plugins/VirtualProduction/Switchboard/Source/Switchboard/switchboard/devices/ndisplay/plugin_ndisplay.py
--- a/Engine/Plugins/VirtualProduction/Switchboard/Source/Switchboard/switchboard/devices/ndisplay/plugin_ndisplay.py
+++ b/Engine/Plugins/VirtualProduction/Switchboard/Source/Switchboard/switchboard/devices/ndisplay/plugin_ndisplay.py
@@ -856,7 +856,7 @@ class DevicenDisplay(DeviceUnreal):
                 '-CONCERTAUTOCONNECT'])
 
         args.extend([
-            f'-CONCERTSERVER="{CONFIG.MUSERVER_SERVER_NAME}"',
+            f'-CONCERTSERVER="{CONFIG.MUSERVER_SERVER_NAME.get_value()}"',
             f'-CONCERTSESSION="{SETTINGS.MUSERVER_SESSION_NAME}"',
             f'-CONCERTDISPLAYNAME="{self.name}"',
             '-CONCERTISHEADLESS',
diff --git a/Engine/Plugins/VirtualProduction/Switchboard/Source/Switchboard/switchboard/devices/unreal/plugin_unreal.py b/Engine/Plugins/VirtualProduction/Switchboard/Source/Switchboard/switchboard/devices/unreal/plugin_unreal.py
--- a/Engine/Plugins/VirtualProduction/Switchboard/Source/Switchboard/switchboard/devices/unreal/plugin_unreal.py
+++ b/Engine/Plugins/VirtualProduction/Switchboard/Source/Switchboard/switchboard/devices/unreal/plugin_unreal.py
@@ -429,7 +429,7 @@ class DeviceUnreal(Device):
             (ipaddress.ip_address(d.ip_address), d)
             for d in cls.active_unreal_devices]
 
-        switchboard_ip_address = ipaddress.ip_address(SETTINGS.IP_ADDRESS)
+        switchboard_ip_address = ipaddress.ip_address(SETTINGS.IP_ADDRESS.get_value())
 
         def is_local(pair):
             return pair[0].is_loopback or (pair[0] == switchboard_ip_address)
@@ -735,7 +735,7 @@ class DeviceUnreal(Device):
 
         # Also check the multi-user server settings.
         muserver_extra_args_lower: str = (
-            CONFIG.MUSERVER_COMMAND_LINE_ARGUMENTS.lower())
+            CONFIG.MUSERVER_COMMAND_LINE_ARGUMENTS.get_value().lower())
         if '-udpmessaging_transport_unicast' in muserver_extra_args_lower:
             LOGGER.error(
                 f'{self.name}: Multi-user server command line arguments '
@@ -1061,7 +1061,7 @@ class DeviceUnreal(Device):
         if (self.is_designated_local_builder() and
                 CONFIG.BUILD_ENGINE.get_value()):
             # Build multi-user server
-            if CONFIG.MUSERVER_AUTO_BUILD:
+            if CONFIG.MUSERVER_AUTO_BUILD.get_value():
                 if DeviceUnreal.mu_server.is_running():
                     mb_ret = QtWidgets.QMessageBox.question(
                         None, 'Terminate multi-user server?',
@@ -1233,10 +1233,10 @@ class DeviceUnreal(Device):
 
     @classmethod
     def get_muserver_endpoint(cls) -> str:
-        setting_val = CONFIG.MUSERVER_ENDPOINT.strip()
+        setting_val = CONFIG.MUSERVER_ENDPOINT.get_value().strip()
         if setting_val:
             return sb_utils.expand_endpoint(setting_val,
-                                            SETTINGS.IP_ADDRESS.strip())
+                                            SETTINGS.IP_ADDRESS.get_value().strip())
         else:
             return ''
 
@@ -1252,7 +1252,7 @@ class DeviceUnreal(Device):
         endpoints: List[str] = []
 
         # Multi-user server.
-        if CONFIG.MUSERVER_AUTO_ENDPOINT:
+        if CONFIG.MUSERVER_AUTO_ENDPOINT.get_value():
             endpoints.append(DeviceUnreal.get_muserver_endpoint())
 
         # Any additional endpoints manually specified via settings.
@@ -1279,7 +1279,7 @@ class DeviceUnreal(Device):
                 '-CONCERTRETRYAUTOCONNECTONERROR '
                 '-CONCERTAUTOCONNECT ')
 
-        command_line_args += (f'-CONCERTSERVER="{CONFIG.MUSERVER_SERVER_NAME}" '
+        command_line_args += (f'-CONCERTSERVER="{CONFIG.MUSERVER_SERVER_NAME.get_value()}" '
                               f'-CONCERTSESSION="{SETTINGS.MUSERVER_SESSION_NAME}" '
                               f'-CONCERTDISPLAYNAME="{self.name}"')
 
@@ -1378,7 +1378,9 @@ class DeviceUnreal(Device):
             return
 
         # Launch the MU server
-        if CONFIG.MUSERVER_AUTO_JOIN.get_value() and self.autojoin_mu_server.get_value() and CONFIG.MUSERVER_AUTO_LAUNCH:
+        if CONFIG.MUSERVER_AUTO_JOIN.get_value() \
+            and self.autojoin_mu_server.get_value() \
+            and CONFIG.MUSERVER_AUTO_LAUNCH.get_value():
             DeviceUnreal.mu_server.launch()
 
         self.compute_runtime_str()
@@ -1424,7 +1426,7 @@ class DeviceUnreal(Device):
         if self.status == DeviceStatus.OPEN:
             self.send_osc_message(
                 osc.OSC_ADD_SEND_TARGET,
-                [SETTINGS.IP_ADDRESS, CONFIG.OSC_SERVER_PORT.get_value()])
+                [SETTINGS.IP_ADDRESS.get_value(), CONFIG.OSC_SERVER_PORT.get_value()])
         else:
             self.osc_connection_timer.stop()
 
@@ -1875,7 +1877,7 @@ class DeviceUnreal(Device):
             DeviceUnreal.rsync_server.make_cygdrive_path(remote_path)
 
         dest_endpoint = \
-            f'{SETTINGS.IP_ADDRESS}:{DeviceUnreal.rsync_server.port}'
+            f'{SETTINGS.IP_ADDRESS.get_value()}:{DeviceUnreal.rsync_server.port}'
         dest_module = DeviceUnreal.rsync_server.INCOMING_LOGS_MODULE
         dest_path = f'rsync://{dest_endpoint}/{dest_module}/'
 
diff --git a/Engine/Plugins/VirtualProduction/Switchboard/Source/Switchboard/switchboard/recording.py b/Engine/Plugins/VirtualProduction/Switchboard/Source/Switchboard/switchboard/recording.py
--- a/Engine/Plugins/VirtualProduction/Switchboard/Source/Switchboard/switchboard/recording.py
+++ b/Engine/Plugins/VirtualProduction/Switchboard/Source/Switchboard/switchboard/recording.py
@@ -85,9 +85,9 @@ class RecordingManager(QtCore.QObject):
         date_string = switchboard_utils.date_to_string(datetime.date.today())
 
         if sequence:
-            return os.path.join(project_dir, CONFIG.PROJECT_NAME, date_string, sequence)
+            return os.path.join(project_dir, CONFIG.PROJECT_NAME.get_value(), date_string, sequence)
 
-        return os.path.join(project_dir, CONFIG.PROJECT_NAME, date_string)
+        return os.path.join(project_dir, CONFIG.PROJECT_NAME.get_value(), date_string)
 
     def add_recording(self, recording):
         # Add a new recording to the manager
@@ -283,7 +283,7 @@ class TransportQueue(QtCore.QObject):
     def run_transport_job(self, transport_job, device):
         # Add Sequence
         take_name = switchboard_utils.capture_name(transport_job.slate, transport_job.take)
-        output_path = os.path.join(SETTINGS.TRANSPORT_PATH, CONFIG.PROJECT_NAME, transport_job.date, take_name)
+        output_path = os.path.join(SETTINGS.TRANSPORT_PATH.get_value(), CONFIG.PROJECT_NAME.get_value(), transport_job.date, take_name)
 
         # Make sure the dirs exist
         os.makedirs(output_path, exist_ok=True)
@@ -295,10 +295,10 @@ class TransportQueue(QtCore.QObject):
         self.active_jobs.append(transport_job)
 
         # Start the thread
-        transport_thread = threading.Thread(target=self._transport_file, args=[device, transport_job.paths, output_path, transport_job.job_name])
+        transport_thread = threading.Thread(target=self._transport_file, args=[device, transport_job.paths, output_path, transport_job, transport_job.job_name])
         transport_thread.start()
 
-    def _transport_file(self, device, device_paths, output_path, job_name):
+    def _transport_file(self, device, device_paths, output_path, transport_job, job_name):
         LOGGER.debug('Start Job')
         self.signal_transport_queue_job_started.emit(transport_job)
 
@@ -316,7 +316,7 @@ class TransportQueue(QtCore.QObject):
         self.signal_transport_queue_job_finished.emit(transport_job)
 
     def file_path(self):
-        project_dir = os.path.join(self._root_dir, 'projects', CONFIG.PROJECT_NAME)
+        project_dir = os.path.join(self._root_dir, 'projects', CONFIG.PROJECT_NAME.get_value())
         os.makedirs(project_dir, exist_ok=True)
 
         return os.path.join(project_dir, TRANSPORT_QUEUE_FILE_NAME)
diff --git a/Engine/Plugins/VirtualProduction/Switchboard/Source/Switchboard/switchboard/settings_dialog.py b/Engine/Plugins/VirtualProduction/Switchboard/Source/Switchboard/switchboard/settings_dialog.py
--- a/Engine/Plugins/VirtualProduction/Switchboard/Source/Switchboard/switchboard/settings_dialog.py
+++ b/Engine/Plugins/VirtualProduction/Switchboard/Source/Switchboard/switchboard/settings_dialog.py
@@ -2,6 +2,7 @@
 
 import os
 import pathlib
+import typing
 
 from PySide2 import QtCore
 from PySide2 import QtGui
@@ -10,7 +11,7 @@ from PySide2 import QtWidgets
 
 from switchboard import config
 from switchboard import switchboard_widgets as sb_widgets
-from switchboard.config import CONFIG, SETTINGS
+from switchboard.config import CONFIG, SETTINGS, Config, UserSettings, Setting
 from switchboard.settings_search import SettingsSearch
 from switchboard.switchboard_widgets import CollapsibleGroupBox, DropDownMenuComboBox
 from switchboard.ui.horizontal_tabs import HorizontalTabWidget
@@ -26,7 +27,7 @@ COLLAPSE_ALL_TEXT = "Collapse all"
 EXPAND_ALL_TEXT = "Expand all"
 
 class SettingsDialog(QtCore.QObject):
-    def __init__(self):
+    def __init__(self, settings: UserSettings, config: Config):
         super().__init__()
         self.plugin_widgets = {}
 
@@ -41,9 +42,9 @@ class SettingsDialog(QtCore.QObject):
         self._create_search_area(dialog_layout)
         self.general_settings_list = [
             self._create_config_path_settings(),
-            self._create_switchboard_settings(),
-            self._create_project_settings(),
-            self._create_multi_user_server_settings()
+            self._create_switchboard_settings(settings, config),
+            self._create_project_settings(config),
+            self._create_multi_user_server_settings(config)
         ]
         
         self._create_tab_widget(dialog_layout)
@@ -194,158 +195,41 @@ class SettingsDialog(QtCore.QObject):
         
         return self.ui.config_path_layout;
         
-    def _create_switchboard_settings(self):
+    def _create_switchboard_settings(self, settings: UserSettings, config: Config):
         self.ui.switchboard_settings_group = CollapsibleGroupBox()
         self.ui.switchboard_settings_group.setTitle("Switchboard")
         layout = QtWidgets.QFormLayout(self.ui.switchboard_settings_group)
         
-        # IP address
-        self.ui.ip_address_label = QtWidgets.QLabel()
-        self.ui.ip_address_label.setText("IP Address")
-        self.ui.ip_address_line_edit = QtWidgets.QLineEdit()
-        layout.addRow(self.ui.ip_address_label, self.ui.ip_address_line_edit)
-        
-        # Transport path
-        self.ui.transport_path_root = QtWidgets.QWidget()
-        self.ui.transport_path_label = QtWidgets.QLabel()
-        self.ui.transport_path_label.setText("Transport Path")
-        transport_path_horizontal_layout = QtWidgets.QHBoxLayout(self.ui.transport_path_root)
-        transport_path_horizontal_layout.setContentsMargins(0, 0, 0, 0)
-        self.ui.transport_path_line_edit = QtWidgets.QLineEdit()
-        self.ui.transport_path_browse_button = QtWidgets.QPushButton()
-        self.ui.transport_path_browse_button.setText("Browse")
-        transport_path_horizontal_layout.addWidget(self.ui.transport_path_line_edit)
-        transport_path_horizontal_layout.addWidget(self.ui.transport_path_browse_button)
-        layout.addRow(self.ui.transport_path_label, self.ui.transport_path_root)
-        
-        # Listener Executable name
-        self.ui.listener_exe_label = QtWidgets.QLabel()
-        self.ui.listener_exe_label.setText("Listener Executable name")
-        self.ui.listener_exe_line_edit = QtWidgets.QLineEdit()
-        layout.addRow(self.ui.listener_exe_label, self.ui.listener_exe_line_edit)
+        settings.IP_ADDRESS.create_ui(form_layout=layout)
+        settings.TRANSPORT_PATH.create_ui(form_layout=layout)
+        config.LISTENER_EXE.create_ui(form_layout=layout)
         
         return self.ui.switchboard_settings_group
         
-    def _create_project_settings(self):
+    def _create_project_settings(self, config: Config):
         self.ui.project_settings_group = CollapsibleGroupBox()
         self.ui.project_settings_group.setTitle("Project Settings")
         layout = QtWidgets.QVBoxLayout(self.ui.project_settings_group)
         
         project_settings_root = QtWidgets.QWidget()
         form_layout = QtWidgets.QFormLayout(project_settings_root)
+        self._create_settings_section(config.basic_project_settings, form_layout)
         layout.addWidget(project_settings_root)
-
-        # Project Name
-        self.ui.project_name_label = QtWidgets.QLabel()
-        self.ui.project_name_label.setText("Project Name")
-        self.ui.project_name_line_edit = QtWidgets.QLineEdit()
-        form_layout.addRow(self.ui.project_name_label, self.ui.project_name_line_edit)
-        
-        # UProject
-        self.ui.uproject_root = QtWidgets.QWidget()
-        self.ui.uproject_label = QtWidgets.QLabel()
-        self.ui.uproject_label.setText("UProject")
-        uproject_horizontal_layout = QtWidgets.QHBoxLayout(self.ui.uproject_root)
-        uproject_horizontal_layout.setContentsMargins(0, 0, 0, 0)
-        self.ui.uproject_line_edit = QtWidgets.QLineEdit()
-        self.ui.uproject_browse_button = QtWidgets.QPushButton()
-        self.ui.uproject_browse_button.setText("Browse")
-        uproject_horizontal_layout.addWidget(self.ui.uproject_line_edit)
-        uproject_horizontal_layout.addWidget(self.ui.uproject_browse_button)
-        form_layout.addRow(self.ui.uproject_label, self.ui.uproject_root)
-        
-        # Engine Dir
-        self.ui.engine_dir_root = QtWidgets.QWidget()
-        self.ui.engine_dir_label = QtWidgets.QLabel()
-        self.ui.engine_dir_label.setText("Engine Dir")
-        engine_dir_horizontal_layout = QtWidgets.QHBoxLayout(self.ui.engine_dir_root)
-        engine_dir_horizontal_layout.setContentsMargins(0, 0, 0, 0)
-        self.ui.engine_dir_line_edit = QtWidgets.QLineEdit()
-        self.ui.engine_dir_browse_button = QtWidgets.QPushButton()
-        self.ui.engine_dir_browse_button.setText("Browse")
-        engine_dir_horizontal_layout.addWidget(self.ui.engine_dir_line_edit)
-        engine_dir_horizontal_layout.addWidget(self.ui.engine_dir_browse_button)
-        form_layout.addRow(self.ui.engine_dir_label, self.ui.engine_dir_root)
-        
-        # Build Engine
-        self.ui.build_engine_label = QtWidgets.QLabel()
-        self.ui.build_engine_label.setText("Build Engine")
-        self.ui.build_engine_checkbox = QtWidgets.QCheckBox()
-        form_layout.addRow(self.ui.build_engine_label, self.ui.build_engine_checkbox)
-
-        # Map Path
-        self.ui.map_path_label = QtWidgets.QLabel()
-        self.ui.map_path_label.setText("Map Path")
-        self.ui.map_path_line_edit = QtWidgets.QLineEdit()
-        form_layout.addRow(self.ui.map_path_label, self.ui.map_path_line_edit)
-        
-        # Map Filter
-        self.ui.map_filter_label = QtWidgets.QLabel()
-        self.ui.map_filter_label.setText("Project Filter")
-        self.ui.map_filter_line_edit = QtWidgets.QLineEdit()
-        form_layout.addRow(self.ui.map_filter_label, self.ui.map_filter_line_edit)
-        
-        # Set up listeners
-        self.ui.uproject_browse_button.clicked.connect(
-            self.uproject_browse_button_clicked)
-        self.ui.engine_dir_browse_button.clicked.connect(
-            self.engine_dir_browse_button_clicked)
-
-        # Update settings in CONFIG when they are changed in the SettingsDialog
-        self.ui.engine_dir_line_edit.editingFinished.connect(
-            lambda widget=self.ui.engine_dir_line_edit:
-            CONFIG.ENGINE_DIR.update_value(widget.text()))
-        self.ui.build_engine_checkbox.stateChanged.connect(
-            lambda state:
-            CONFIG.BUILD_ENGINE.update_value(
-                True if state == QtCore.Qt.Checked else False))
-        self.ui.uproject_line_edit.editingFinished.connect(
-            lambda widget=self.ui.uproject_line_edit:
-            CONFIG.UPROJECT_PATH.update_value(widget.text()))
-        self.ui.map_path_line_edit.editingFinished.connect(
-            lambda widget=self.ui.map_path_line_edit:
-            CONFIG.MAPS_PATH.update_value(widget.text()))
-        self.ui.map_filter_line_edit.editingFinished.connect(
-            lambda widget=self.ui.map_filter_line_edit:
-            CONFIG.MAPS_FILTER.update_value(widget.text()))
         
         # Sub settings
-        self._create_osc_settings(layout)
-        self._create_source_control_settings(layout)
+        self._create_osc_settings(config, layout)
+        self._create_source_control_settings(config, layout)
+        
         return self.ui.project_settings_group
     
-    def _create_osc_settings(self, layout):
+    def _create_osc_settings(self, config: Config, layout: QtWidgets.QFormLayout):
         self.ui.osc_settings_group = QtWidgets.QGroupBox()
         self.ui.osc_settings_group.setTitle("OSC")
         form_layout = QtWidgets.QFormLayout(self.ui.osc_settings_group)
         layout.addWidget(self.ui.osc_settings_group)
-        
-        # Server port
-        self.ui.osc_server_port_label = QtWidgets.QLabel()
-        self.ui.osc_server_port_label.setText("Server Port")
-        self.ui.osc_server_port_line_edit = QtWidgets.QLineEdit()
-        form_layout.addRow(self.ui.osc_server_port_label, self.ui.osc_server_port_line_edit)
-        
-        # Client port
-        self.ui.osc_client_port_label = QtWidgets.QLabel()
-        self.ui.osc_client_port_label.setText("Client Port")
-        self.ui.osc_client_port_line_edit = QtWidgets.QLineEdit()
-        form_layout.addRow(self.ui.osc_client_port_label, self.ui.osc_client_port_line_edit)
-        
-        # Set up listeners
-        max_port = (1 << 16) - 1
-        self.ui.osc_server_port_line_edit.setValidator(
-            QtGui.QIntValidator(0, max_port))
-        self.ui.osc_client_port_line_edit.setValidator(
-            QtGui.QIntValidator(0, max_port))
-        self.ui.osc_server_port_line_edit.editingFinished.connect(
-            lambda widget=self.ui.osc_server_port_line_edit:
-            CONFIG.OSC_SERVER_PORT.update_value(int(widget.text())))
-        self.ui.osc_client_port_line_edit.editingFinished.connect(
-            lambda widget=self.ui.osc_client_port_line_edit:
-            CONFIG.OSC_CLIENT_PORT.update_value(int(widget.text())))
+        self._create_settings_section(config.osc_settings, form_layout)
     
-    def _create_source_control_settings(self, layout):
+    def _create_source_control_settings(self, config: Config, layout: QtWidgets.QFormLayout):
         self.ui.source_control_settings_group = QtWidgets.QGroupBox()
         self.ui.source_control_settings_group.setTitle("Source Control")
         self.ui.source_control_settings_group.setCheckable(True)
@@ -357,97 +241,20 @@ class SettingsDialog(QtCore.QObject):
         )
         form_layout = QtWidgets.QFormLayout(self.ui.source_control_settings_group)
         layout.addWidget(self.ui.source_control_settings_group)
-        
-        # P4 Project Path
-        self.ui.p4_project_path_label = QtWidgets.QLabel()
-        self.ui.p4_project_path_label.setText("P4 Project Path")
-        self.ui.p4_project_path_line_edit = QtWidgets.QLineEdit()
-        form_layout.addRow(self.ui.p4_project_path_label, self.ui.p4_project_path_line_edit)
-        
-        # P4 Engine Path
-        self.ui.p4_engine_path_label = QtWidgets.QLabel()
-        self.ui.p4_engine_path_label.setText("P4 Engine Path")
-        self.ui.p4_engine_path_line_edit = QtWidgets.QLineEdit()
-        form_layout.addRow(self.ui.p4_engine_path_label, self.ui.p4_engine_path_line_edit)
-        
-        # P4 Workspace name
-        self.ui.source_control_workspace_label = QtWidgets.QLabel()
-        self.ui.source_control_workspace_label.setText("P4 Workspace name")
-        self.ui.source_control_workspace_line_edit = QtWidgets.QLineEdit()
-        form_layout.addRow(self.ui.source_control_workspace_label, self.ui.source_control_workspace_line_edit)
-        
-        # Set up listeners
-        self.ui.p4_project_path_line_edit.editingFinished.connect(
-            lambda widget=self.ui.p4_project_path_line_edit:
-            CONFIG.P4_PROJECT_PATH.update_value(widget.text()))
-        self.ui.p4_engine_path_line_edit.editingFinished.connect(
-            lambda widget=self.ui.p4_engine_path_line_edit:
-            CONFIG.P4_ENGINE_PATH.update_value(widget.text()))
-        self.ui.source_control_workspace_line_edit.editingFinished.connect(
-            lambda widget=self.ui.source_control_workspace_line_edit:
-            CONFIG.SOURCE_CONTROL_WORKSPACE.update_value(widget.text()))
+        self._create_settings_section(config.source_control_settings, form_layout)
     
-    def _create_multi_user_server_settings(self):
+    def _create_multi_user_server_settings(self, config: Config):
         self.ui.multi_user_settings = CollapsibleGroupBox()
         self.ui.multi_user_settings.setTitle("Multi User Server")
         form_layout = QtWidgets.QFormLayout(self.ui.multi_user_settings)
         
-        # Server Name
-        self.ui.mu_server_name_label = QtWidgets.QLabel()
-        self.ui.mu_server_name_label.setText("Server Name")
-        self.ui.mu_server_name_line_edit = QtWidgets.QLineEdit()
-        self.ui.mu_server_name_line_edit.setPlaceholderText("MU_Server")
-        form_layout.addRow(self.ui.mu_server_name_label, self.ui.mu_server_name_line_edit)
-        
-        # Command Line Args
-        self.ui.mu_cmd_line_args_label = QtWidgets.QLabel()
-        self.ui.mu_cmd_line_args_label.setText("Command Line Args")
-        self.ui.mu_cmd_line_args_line_edit = QtWidgets.QLineEdit()
-        form_layout.addRow(self.ui.mu_cmd_line_args_label, self.ui.mu_cmd_line_args_line_edit)
-        
-        # Unicast Endpoint
-        self.ui.mu_server_endpoint_label = QtWidgets.QLabel()
-        self.ui.mu_server_endpoint_label.setText("Unicast Endpoint")
-        self.ui.mu_server_endpoint_line_edit = QtWidgets.QLineEdit()
-        form_layout.addRow(self.ui.mu_server_endpoint_label, self.ui.mu_server_endpoint_line_edit)
-        
-        # Multicast Endpoint
-        self.ui.mu_server_multicast_endpoint_label = QtWidgets.QLabel()
-        self.ui.mu_server_multicast_endpoint_label.setText("Multicast Endpoint")
-        self.ui.mu_server_multicast_endpoint_line_edit = QtWidgets.QLineEdit()
-        form_layout.addRow(self.ui.mu_server_multicast_endpoint_label, self.ui.mu_server_multicast_endpoint_line_edit)
-        
-        # Clean History
-        self.ui.mu_clean_history_label = QtWidgets.QLabel()
-        self.ui.mu_clean_history_label.setText("Clean History")
-        self.ui.mu_clean_history_check_box = QtWidgets.QCheckBox()
-        form_layout.addRow(self.ui.mu_clean_history_label, self.ui.mu_clean_history_check_box)
-        
-        # Auto Launch
-        self.ui.mu_auto_launch_label = QtWidgets.QLabel()
-        self.ui.mu_auto_launch_label.setText("Auto Launch")
-        self.ui.mu_auto_launch_check_box = QtWidgets.QCheckBox()
-        form_layout.addRow(self.ui.mu_auto_launch_label, self.ui.mu_auto_launch_check_box)
-        
-        # Executable Name
-        self.ui.muserver_exe_label = QtWidgets.QLabel()
-        self.ui.muserver_exe_label.setText("Executable Name")
-        self.ui.muserver_exe_line_edit = QtWidgets.QLineEdit()
-        form_layout.addRow(self.ui.muserver_exe_label, self.ui.muserver_exe_line_edit)
-        
-        # Auto Build
-        self.ui.muserver_auto_build_label = QtWidgets.QLabel()
-        self.ui.muserver_auto_build_label.setText("Auto Build")
-        self.ui.muserver_auto_build_check_box = QtWidgets.QCheckBox()
-        form_layout.addRow(self.ui.muserver_auto_build_label, self.ui.muserver_auto_build_check_box)
-        
-        # Auto Static Endpoint
-        self.ui.muserver_auto_endpoint_label = QtWidgets.QLabel()
-        self.ui.muserver_auto_endpoint_label.setText("Auto Static Endpoint")
-        self.ui.muserver_auto_endpoint_check_box = QtWidgets.QCheckBox()
-        form_layout.addRow(self.ui.muserver_auto_endpoint_label, self.ui.muserver_auto_endpoint_check_box)
-        
+        self._create_settings_section(config.mu_settings, form_layout)
         return self.ui.multi_user_settings
+    
+    def _create_settings_section(self, dict: typing.Dict[str, Setting], layout: QtWidgets.QFormLayout):
+        setting_list = [setting for _, setting in dict.items()]
+        for setting in setting_list:
+            setting.create_ui(override_device_name=None, form_layout=layout)
 
     def _create_tab_widget(self, parent_layout):
         self.ui.tab_widget = HorizontalTabWidget()
@@ -521,158 +328,6 @@ class SettingsDialog(QtCore.QObject):
 
         self._changed_config_path = config_path
 
-    def uproject_browse_button_clicked(self):
-        file_name, _ = QtWidgets.QFileDialog.getOpenFileName(
-            parent=self.ui, filter='*.uproject')
-        if file_name:
-            file_name = os.path.normpath(file_name)
-            self.set_uproject(file_name)
-            CONFIG.UPROJECT_PATH.update_value(file_name)
-
-    def engine_dir_browse_button_clicked(self):
-        dir_name = QtWidgets.QFileDialog.getExistingDirectory(parent=self.ui)
-        if dir_name:
-            dir_name = os.path.normpath(dir_name)
-            self.set_engine_dir(dir_name)
-            CONFIG.ENGINE_DIR.update_value(dir_name)
-
-    def ip_address(self):
-        return self.ui.ip_address_line_edit.text()
-
-    def set_ip_address(self, value):
-        self.ui.ip_address_line_edit.setText(value)
-
-    def transport_path(self):
-        return self.ui.transport_path_line_edit.text()
-
-    def set_transport_path(self, value):
-        self.ui.transport_path_line_edit.setText(value)
-
-    def listener_exe(self):
-        return self.ui.listener_exe_line_edit.text()
-
-    def set_listener_exe(self, value):
-        self.ui.listener_exe_line_edit.setText(value)
-
-    def project_name(self):
-        return self.ui.project_name_line_edit.text()
-
-    def set_project_name(self, value):
-        self.ui.project_name_line_edit.setText(value)
-
-    def uproject(self):
-        return self.ui.uproject_line_edit.text()
-
-    def set_uproject(self, value):
-        self.ui.uproject_line_edit.setText(value)
-
-    def engine_dir(self):
-        return self.ui.engine_dir_line_edit.text()
-
-    def set_engine_dir(self, value):
-        self.ui.engine_dir_line_edit.setText(value)
-
-    def set_build_engine(self, value):
-        self.ui.build_engine_checkbox.setChecked(value)
-
-    def p4_enabled(self):
-        return self.ui.source_control_settings_group.isChecked()
-
-    def set_p4_enabled(self, enabled):
-        self.ui.source_control_settings_group.setChecked(enabled)
-
-    def p4_project_path(self):
-        return self.ui.p4_project_path_line_edit.text()
-
-    def set_p4_project_path(self, value):
-        self.ui.p4_project_path_line_edit.setText(value)
-
-    def p4_engine_path(self):
-        return self.ui.p4_engine_path_line_edit.text()
-
-    def set_p4_engine_path(self, value):
-        self.ui.p4_engine_path_line_edit.setText(value)
-
-    def source_control_workspace(self):
-        return self.ui.source_control_workspace_line_edit.text()
-
-    def set_source_control_workspace(self, value):
-        self.ui.source_control_workspace_line_edit.setText(value)
-
-    def map_path(self):
-        return self.ui.map_path_line_edit.text()
-
-    def set_map_path(self, value):
-        self.ui.map_path_line_edit.setText(value)
-
-    def map_filter(self):
-        return self.ui.map_filter_line_edit.text()
-
-    def set_map_filter(self, value):
-        self.ui.map_filter_line_edit.setText(value)
-
-    # OSC Settings
-    def set_osc_server_port(self, port):
-        self.ui.osc_server_port_line_edit.setText(str(port))
-
-    def set_osc_client_port(self, port):
-        self.ui.osc_client_port_line_edit.setText(str(port))
-
-    # MU SERVER Settings
-    def mu_server_name(self):
-        return self.ui.mu_server_name_line_edit.text()
-
-    def set_mu_server_name(self, value):
-        self.ui.mu_server_name_line_edit.setText(value)
-
-    def mu_server_multicast_endpoint(self):
-        return self.ui.mu_server_multicast_endpoint_line_edit.text()
-
-    def set_mu_server_multicast_endpoint(self, endpoint):
-        self.ui.mu_server_multicast_endpoint_line_edit.setText(str(endpoint))
-
-    def mu_server_endpoint(self):
-        return self.ui.mu_server_endpoint_line_edit.text()
-
-    def set_mu_server_endpoint(self, endpoint):
-        self.ui.mu_server_endpoint_line_edit.setText(str(endpoint))
-
-    def mu_cmd_line_args(self):
-        return self.ui.mu_cmd_line_args_line_edit.text()
-
-    def set_mu_cmd_line_args(self, value):
-        self.ui.mu_cmd_line_args_line_edit.setText(value)
-
-    def mu_auto_launch(self):
-        return self.ui.mu_auto_launch_check_box.isChecked()
-
-    def set_mu_auto_launch(self, value):
-        self.ui.mu_auto_launch_check_box.setChecked(value)
-
-    def mu_clean_history(self):
-        return self.ui.mu_clean_history_check_box.isChecked()
-
-    def set_mu_clean_history(self, value):
-        self.ui.mu_clean_history_check_box.setChecked(value)
-
-    def mu_server_exe(self):
-        return self.ui.muserver_exe_line_edit.text()
-
-    def set_mu_server_exe(self, value):
-        self.ui.muserver_exe_line_edit.setText(value)
-
-    def mu_server_auto_build(self) -> bool:
-        return self.ui.muserver_auto_build_check_box.isChecked()
-
-    def set_mu_server_auto_build(self, value: bool):
-        self.ui.muserver_auto_build_check_box.setChecked(value)
-
-    def mu_server_auto_endpoint(self) -> bool:
-        return self.ui.muserver_auto_endpoint_check_box.isChecked()
-
-    def set_mu_server_auto_endpoint(self, value: bool):
-        self.ui.muserver_auto_endpoint_check_box.setChecked(value)
-
     # Devices
     def add_section_for_plugin(self, plugin_name, plugin_settings, device_settings):
         any_device_settings = (
diff --git a/Engine/Plugins/VirtualProduction/Switchboard/Source/Switchboard/switchboard/switchboard_application.py b/Engine/Plugins/VirtualProduction/Switchboard/Source/Switchboard/switchboard/switchboard_application.py
--- a/Engine/Plugins/VirtualProduction/Switchboard/Source/Switchboard/switchboard/switchboard_application.py
+++ b/Engine/Plugins/VirtualProduction/Switchboard/Source/Switchboard/switchboard/switchboard_application.py
@@ -105,10 +105,10 @@ class MultiUserApplication:
         return ''
 
     def get_mu_server_endpoint_arg(self):
-        setting_val = CONFIG.MUSERVER_ENDPOINT.strip()
+        setting_val = CONFIG.MUSERVER_ENDPOINT.get_value().strip()
         if setting_val:
             converted_ip = sb_utils.expand_endpoint(setting_val,
-                                                    SETTINGS.IP_ADDRESS.strip())
+                                                    SETTINGS.IP_ADDRESS.get_value().strip())
             return f'-UDPMESSAGING_TRANSPORT_UNICAST="{converted_ip}"'
         return ''
 
@@ -132,14 +132,14 @@ class MultiUserApplication:
 
             cmdline = f'start "Multi User Server" "{self.exe_path()}"'
             cmdline += f' -CONCERTSERVER="{CONFIG.MUSERVER_SERVER_NAME}"'
-            cmdline += f' {CONFIG.MUSERVER_COMMAND_LINE_ARGUMENTS}'
+            cmdline += f' {CONFIG.MUSERVER_COMMAND_LINE_ARGUMENTS.get_value()}'
             cmdline += f' {self.get_mu_server_endpoint_arg()}'
             cmdline += f' {self.get_mu_server_multicast_arg()}'
 
             if self.concert_ignore_cl:
                 cmdline += " -ConcertIgnore"
 
-            if CONFIG.MUSERVER_CLEAN_HISTORY:
+            if CONFIG.MUSERVER_CLEAN_HISTORY.get_value():
                 cmdline += " -ConcertClean"
 
             if len(args) > 0:
@@ -267,7 +267,7 @@ hosts allow = {allowed_addrs}
             client.address for client in self.allowed_clients.values()}
 
         allowed_addrs.add(self.address)
-        allowed_addrs.add(SETTINGS.IP_ADDRESS)
+        allowed_addrs.add(SETTINGS.IP_ADDRESS.get_value())
         allowed_addrs.discard('0.0.0.0')
 
         config = self.CONFIG_TEMPLATE.format(
diff --git a/Engine/Plugins/VirtualProduction/Switchboard/Source/Switchboard/switchboard/switchboard_dialog.py b/Engine/Plugins/VirtualProduction/Switchboard/Source/Switchboard/switchboard/switchboard_dialog.py
--- a/Engine/Plugins/VirtualProduction/Switchboard/Source/Switchboard/switchboard/switchboard_dialog.py
+++ b/Engine/Plugins/VirtualProduction/Switchboard/Source/Switchboard/switchboard/switchboard_dialog.py
@@ -292,7 +292,7 @@ class SwitchboardDialog(QtCore.QObject):
 
         # Start the OSC server
         self.osc_server = switchboard_application.OscServer()
-        self.osc_server.launch(SETTINGS.IP_ADDRESS, CONFIG.OSC_SERVER_PORT.get_value())
+        self.osc_server.launch(SETTINGS.IP_ADDRESS.get_value(), CONFIG.OSC_SERVER_PORT.get_value())
 
         # Register with OSC server
         self.osc_server.dispatcher_map(osc.TAKE, self.osc_take)
@@ -808,33 +808,7 @@ class SwitchboardDialog(QtCore.QObject):
         Settings window
         """
         # TODO: VALIDATE RECORD PATH
-        settings_dialog = SettingsDialog()
-
-        settings_dialog.set_config_path(SETTINGS.CONFIG)
-        settings_dialog.set_ip_address(SETTINGS.IP_ADDRESS)
-        settings_dialog.set_transport_path(SETTINGS.TRANSPORT_PATH)
-        settings_dialog.set_listener_exe(CONFIG.LISTENER_EXE)
-        settings_dialog.set_project_name(CONFIG.PROJECT_NAME)
-        settings_dialog.set_uproject(CONFIG.UPROJECT_PATH.get_value())
-        settings_dialog.set_engine_dir(CONFIG.ENGINE_DIR.get_value())
-        settings_dialog.set_build_engine(CONFIG.BUILD_ENGINE.get_value())
-        settings_dialog.set_p4_enabled(bool(CONFIG.P4_ENABLED.get_value()))
-        settings_dialog.set_source_control_workspace(CONFIG.SOURCE_CONTROL_WORKSPACE.get_value())
-        settings_dialog.set_p4_project_path(CONFIG.P4_PROJECT_PATH.get_value())
-        settings_dialog.set_p4_engine_path(CONFIG.P4_ENGINE_PATH.get_value())
-        settings_dialog.set_map_path(CONFIG.MAPS_PATH.get_value())
-        settings_dialog.set_map_filter(CONFIG.MAPS_FILTER.get_value())
-        settings_dialog.set_osc_server_port(CONFIG.OSC_SERVER_PORT.get_value())
-        settings_dialog.set_osc_client_port(CONFIG.OSC_CLIENT_PORT.get_value())
-        settings_dialog.set_mu_server_name(CONFIG.MUSERVER_SERVER_NAME)
-        settings_dialog.set_mu_server_endpoint(CONFIG.MUSERVER_ENDPOINT)
-        settings_dialog.set_mu_server_multicast_endpoint(CONFIG.MUSERVER_MULTICAST_ENDPOINT.get_value())
-        settings_dialog.set_mu_cmd_line_args(CONFIG.MUSERVER_COMMAND_LINE_ARGUMENTS)
-        settings_dialog.set_mu_clean_history(CONFIG.MUSERVER_CLEAN_HISTORY)
-        settings_dialog.set_mu_auto_launch(CONFIG.MUSERVER_AUTO_LAUNCH)
-        settings_dialog.set_mu_server_exe(CONFIG.MULTIUSER_SERVER_EXE)
-        settings_dialog.set_mu_server_auto_build(CONFIG.MUSERVER_AUTO_BUILD)
-        settings_dialog.set_mu_server_auto_endpoint(CONFIG.MUSERVER_AUTO_ENDPOINT)
+        settings_dialog = SettingsDialog(SETTINGS, CONFIG)
 
         for plugin_name in sorted(self.device_manager.available_device_plugins(), key=str.lower):
             device_instances = self.device_manager.devices_of_type(plugin_name)
@@ -843,6 +817,8 @@ class SwitchboardDialog(QtCore.QObject):
         
         settings_dialog.select_all_tab()
 
+        old_ip_address = SETTINGS.IP_ADDRESS.get_value()
+        old_transport_path = SETTINGS.TRANSPORT_PATH.get_value()
         # avoid saving the config all the time while in the settings dialog
         CONFIG.push_saving_allowed(False)
         try:
@@ -855,73 +831,14 @@ class SwitchboardDialog(QtCore.QObject):
         new_config_path = settings_dialog.config_path()
         if new_config_path != SETTINGS.CONFIG and new_config_path is not None:
             CONFIG.replace(new_config_path)
-
             SETTINGS.CONFIG = new_config_path
             SETTINGS.save()
 
-        ip_address = settings_dialog.ip_address()
-        if ip_address != SETTINGS.IP_ADDRESS:
-            SETTINGS.IP_ADDRESS = ip_address
+        if old_ip_address != SETTINGS.IP_ADDRESS.get_value() or SETTINGS.TRANSPORT_PATH.get_value():
             SETTINGS.save()
-
-            # Relaunch the OSC server
+        if old_ip_address != SETTINGS.IP_ADDRESS.get_value():
             self.osc_server.close()
-            self.osc_server.launch(SETTINGS.IP_ADDRESS, CONFIG.OSC_SERVER_PORT.get_value())
-
-        transport_path = settings_dialog.transport_path()
-        if transport_path != SETTINGS.TRANSPORT_PATH:
-            SETTINGS.TRANSPORT_PATH = transport_path
-            SETTINGS.save()
-
-        # todo-dara, when these project settings have been converted into actual Settings
-        # these assignments are not needed anymore, as the settings would be directly connected to their widgets
-
-        project_name = settings_dialog.project_name()
-        if project_name != CONFIG.PROJECT_NAME:
-            CONFIG.PROJECT_NAME = project_name
-
-        listener_exe = settings_dialog.listener_exe()
-        if listener_exe != CONFIG.LISTENER_EXE:
-            CONFIG.LISTENER_EXE = listener_exe
-
-        # Multi User Settings
-        mu_server_name = settings_dialog.mu_server_name()
-        if mu_server_name != CONFIG.MUSERVER_SERVER_NAME:
-            CONFIG.MUSERVER_SERVER_NAME = mu_server_name
-
-        mu_server_endpoint = settings_dialog.mu_server_endpoint()
-        if mu_server_endpoint != CONFIG.MUSERVER_ENDPOINT:
-            CONFIG.MUSERVER_ENDPOINT = mu_server_endpoint
-
-        mu_server_multicast_endpoint = settings_dialog.mu_server_multicast_endpoint()
-        if mu_server_multicast_endpoint != CONFIG.MUSERVER_MULTICAST_ENDPOINT:
-            CONFIG.MUSERVER_MULTICAST_ENDPOINT.update_value(mu_server_multicast_endpoint)
-
-        mu_cmd_line_args = settings_dialog.mu_cmd_line_args()
-        if mu_cmd_line_args != CONFIG.MUSERVER_COMMAND_LINE_ARGUMENTS:
-            CONFIG.MUSERVER_COMMAND_LINE_ARGUMENTS = mu_cmd_line_args
-
-        mu_clean_history = settings_dialog.mu_clean_history()
-        if mu_clean_history != CONFIG.MUSERVER_CLEAN_HISTORY:
-            CONFIG.MUSERVER_CLEAN_HISTORY = mu_clean_history
-
-        mu_auto_launch = settings_dialog.mu_auto_launch()
-        if mu_auto_launch != CONFIG.MUSERVER_AUTO_LAUNCH:
-            CONFIG.MUSERVER_AUTO_LAUNCH = mu_auto_launch
-
-        mu_server_exe = settings_dialog.mu_server_exe()
-        if mu_server_exe != CONFIG.MULTIUSER_SERVER_EXE:
-            CONFIG.MULTIUSER_SERVER_EXE = mu_server_exe
-
-        mu_auto_build = settings_dialog.mu_server_auto_build()
-        if mu_auto_build != CONFIG.MUSERVER_AUTO_BUILD:
-            CONFIG.MUSERVER_AUTO_BUILD = mu_auto_build
-
-        mu_auto_endpoint = settings_dialog.mu_server_auto_endpoint()
-        if mu_auto_endpoint != CONFIG.MUSERVER_AUTO_ENDPOINT:
-            CONFIG.MUSERVER_AUTO_ENDPOINT = mu_auto_endpoint
-
-        CONFIG.P4_ENABLED.update_value(settings_dialog.p4_enabled())
+            self.osc_server.launch(SETTINGS.IP_ADDRESS.get_value(), CONFIG.OSC_SERVER_PORT.get_value())
 
         CONFIG.save()
 
@@ -1864,7 +1781,7 @@ class SwitchboardDialog(QtCore.QObject):
 
         # Return a Recording object
         new_recording = recording.Recording()
-        new_recording.project = CONFIG.PROJECT_NAME
+        new_recording.project = CONFIG.PROJECT_NAME.get_value()
         new_recording.shoot = self.shoot
         new_recording.sequence = self.sequence
         new_recording.slate = self.slate
