Add current IP and config file to bottom right

diff --git a/Engine/Plugins/VirtualProduction/Switchboard/Source/Switchboard/switchboard/switchboard_dialog.py b/Engine/Plugins/VirtualProduction/Switchboard/Source/Switchboard/switchboard/switchboard_dialog.py
--- a/Engine/Plugins/VirtualProduction/Switchboard/Source/Switchboard/switchboard/switchboard_dialog.py
+++ b/Engine/Plugins/VirtualProduction/Switchboard/Source/Switchboard/switchboard/switchboard_dialog.py
@@ -402,6 +402,39 @@ class SwitchboardDialog(QtCore.QObject):
 
         # Run the transport queue
         #self.transport_queue_resume()
+        
+        self.update_current_config_text()
+        self.update_current_ip_address_text()
+        SETTINGS.IP_ADDRESS.signal_setting_changed.connect(
+            lambda: self.update_current_ip_address_text()
+        )
+        self.window.current_ip_value.editingFinished.connect(self._try_change_ip_address)
+        
+    def _try_change_ip_address(self):
+        def is_valid_ip_format(address):
+            ip_port_split = address.split(":")
+            if len(ip_port_split) != 2:
+                return False
+            
+            port = ip_port_split[1]
+            if not port.isdigit():
+                return False
+            
+            ip_elements = ip_port_split[0].split(".")
+            if len(ip_elements) != 4:
+                return False
+            
+            return all(element.isdigit() for element in ip_elements)
+        
+        new_value = self.window.current_ip_value.text()
+        if is_valid_ip_format(new_value):
+            SETTINGS.IP_ADDRESS.update_value(new_value)
+            SETTINGS.save()
+            # Print to make it clear to the user that the IP was updated
+            LOGGER.info(f"Updated IP to {new_value}")
+        else:
+            self.window.current_ip_value.setText(SETTINGS.IP_ADDRESS.get_value())
+            LOGGER.warning(f"{new_value} must be in the format a.b.c.d:port")
 
     def _poll_muserver_status(self):
         '''
@@ -734,6 +767,14 @@ class SwitchboardDialog(QtCore.QObject):
         self.p4_refresh_project_cl()
         self.p4_refresh_engine_cl()
         self.refresh_levels()
+        self.update_current_config_text()
+        
+    def update_current_config_text(self):
+        file_name = os.path.basename(SETTINGS.CONFIG)
+        self.window.current_config_file_value.setText(file_name)
+    
+    def update_current_ip_address_text(self):
+        self.window.current_ip_value.setText(SETTINGS.IP_ADDRESS.get_value())
 
     def menu_new_config(self):
         uproject_search_path = os.path.dirname(CONFIG.UPROJECT_PATH.get_value().replace('"',''))
