Fix port being accepted for ip field

diff --git a/Engine/Plugins/VirtualProduction/Switchboard/Source/Switchboard/switchboard/switchboard_dialog.py b/Engine/Plugins/VirtualProduction/Switchboard/Source/Switchboard/switchboard/switchboard_dialog.py
--- a/Engine/Plugins/VirtualProduction/Switchboard/Source/Switchboard/switchboard/switchboard_dialog.py
+++ b/Engine/Plugins/VirtualProduction/Switchboard/Source/Switchboard/switchboard/switchboard_dialog.py
@@ -413,11 +413,7 @@ class SwitchboardDialog(QtCore.QObject):
     def _try_change_ip_address(self):
         def is_valid_ip_format(address):
             ip_port_split = address.split(":")
-            if len(ip_port_split) != 2:
-                return False
-            
-            port = ip_port_split[1]
-            if not port.isdigit():
+            if len(ip_port_split) != 1:
                 return False
             
             ip_elements = ip_port_split[0].split(".")
@@ -434,7 +430,7 @@ class SwitchboardDialog(QtCore.QObject):
             LOGGER.info(f"Updated IP to {new_value}")
         else:
             self.window.current_ip_value.setText(SETTINGS.IP_ADDRESS.get_value())
-            LOGGER.warning(f"{new_value} must be in the format a.b.c.d:port")
+            LOGGER.warning(f"{new_value} must be in the format a.b.c.d (without port)")
 
     def _poll_muserver_status(self):
         '''
