Fix inaccurate tooltip

diff --git a/Engine/Plugins/VirtualProduction/Switchboard/Source/Switchboard/switchboard/config.py b/Engine/Plugins/VirtualProduction/Switchboard/Source/Switchboard/switchboard/config.py
--- a/Engine/Plugins/VirtualProduction/Switchboard/Source/Switchboard/switchboard/config.py
+++ b/Engine/Plugins/VirtualProduction/Switchboard/Source/Switchboard/switchboard/config.py
@@ -2012,13 +2012,15 @@ class Config(object):
         self.mu_settings = {
             "muserver_server_name": StringSetting(
                 "muserver_server_name",
-                "Command Line Args",
-                data.get('muserver_server_name', f'{self.PROJECT_NAME.get_value()}_MU_Server')
+                "Server name",
+                data.get('muserver_server_name', f'{self.PROJECT_NAME.get_value()}_MU_Server'),
+                tool_tip="The name that will be given to the server"
             ),
             "muserver_command_line_arguments": StringSetting(
                 "muserver_command_line_arguments",
-                "Server name",
-                data.get('muserver_command_line_arguments', '')
+                "Command Line Args",
+                data.get('muserver_command_line_arguments', ''),
+                tool_tip="Additional command line arguments to pass to multiuser"
             ),
             "muserver_endpoint": StringSetting(
                 "muserver_endpoint",
