Handle corrupted files

diff --git a/Engine/Plugins/VirtualProduction/Switchboard/Source/Switchboard/switchboard/__main__.py b/Engine/Plugins/VirtualProduction/Switchboard/Source/Switchboard/switchboard/__main__.py
--- a/Engine/Plugins/VirtualProduction/Switchboard/Source/Switchboard/switchboard/__main__.py
+++ b/Engine/Plugins/VirtualProduction/Switchboard/Source/Switchboard/switchboard/__main__.py
@@ -1,4 +1,5 @@
 # Copyright Epic Games, Inc. All Rights Reserved.
+from switchboard.config import SETTINGS, CONFIG
 from .switchboard_dialog import SwitchboardDialog
 from .switchboard_logging import LOGGER
 
@@ -23,6 +24,11 @@ def launch():
 
     QtWidgets.QApplication.setAttribute(QtCore.Qt.AA_EnableHighDpiScaling)
     app = QtWidgets.QApplication(sys.argv)
+
+    # Global variables can be inited after QT has been started (they may show dialogues if there are errors)
+    SETTINGS.init()
+    CONFIG.init(SETTINGS.CONFIG)
+    
     main_window = SwitchboardDialog()
 
     if not main_window.window:
diff --git a/Engine/Plugins/VirtualProduction/Switchboard/Source/Switchboard/switchboard/config.py b/Engine/Plugins/VirtualProduction/Switchboard/Source/Switchboard/switchboard/config.py
--- a/Engine/Plugins/VirtualProduction/Switchboard/Source/Switchboard/switchboard/config.py
+++ b/Engine/Plugins/VirtualProduction/Switchboard/Source/Switchboard/switchboard/config.py
@@ -5,6 +5,7 @@ import fnmatch
 import json
 import os
 import pathlib
+import shutil
 import socket
 import sys
 import typing
@@ -20,7 +21,10 @@ from switchboard.switchboard_widgets import DropDownMenuComboBox
 ROOT_CONFIGS_PATH = pathlib.Path(__file__).parent.with_name('configs')
 CONFIG_SUFFIX = '.json'
 
-USER_SETTINGS_FILE_PATH = ROOT_CONFIGS_PATH.joinpath('user_settings.json')
+USER_SETTINGS_FILE_NAME = 'user_settings.json'
+USER_SETTINGS_FILE_PATH = ROOT_CONFIGS_PATH.joinpath(USER_SETTINGS_FILE_NAME)
+USER_SETTINGS_BACKUP_FILE_NAME = 'corrupted_user_settings_backup.json'
+USER_SETTINGS_BACKUP_FILE_PATH = ROOT_CONFIGS_PATH.joinpath(USER_SETTINGS_BACKUP_FILE_NAME)
 
 DEFAULT_MAP_TEXT = '-- Default Map --'
 
@@ -1786,8 +1790,8 @@ class Config(object):
         ''' Restores saving_allowed flag from the stack
         '''
         self.saving_allowed = self.saving_allowed_fifo.pop()
-
-    def __init__(self, file_path: typing.Union[str, pathlib.Path]):
+        
+    def init(self, file_path: typing.Union[str, pathlib.Path]):
         self.init_with_file_path(file_path)
 
     def init_with_file_path(self, file_path: typing.Union[str, pathlib.Path]):
@@ -1803,6 +1807,10 @@ class Config(object):
                 LOGGER.error(f'Config: {e}')
                 self.file_path = None
                 data = {}
+            except ValueError:
+                # The original file will be overwritten
+                self._backup_corrupted_config(self.file_path.__str__())
+                data = {}
         else:
             self.file_path = None
             data = {}
@@ -1856,6 +1864,25 @@ class Config(object):
                     self._device_data_from_config.setdefault(
                         device_type, []).append(device_data)
 
+    def _backup_corrupted_config(self, original_file_path: str):
+        directory_name = os.path.dirname(original_file_path)
+        original_file_name = os.path.basename(original_file_path)
+        new_file_name = original_file_name.replace(".", "_corrupted_backup.")
+
+        LOGGER.error(f'{original_file_name} has invalid JSON format. Creating default...')
+        answer = QtWidgets.QMessageBox.question(
+            None,
+            'Invalid project settings',
+            f'Config file { original_file_name } is invalid JSON and will be replaced by a new default JSON config.'
+            f'\n\nDo you want to save a backup named { new_file_name }?',
+            QtWidgets.QMessageBox.Yes | QtWidgets.QMessageBox.No
+        )
+        if answer == QtWidgets.QMessageBox.Yes:
+            new_file_path = os.path.join(directory_name, new_file_name)
+            if os.path.exists(new_file_path):
+                os.remove(new_file_path)
+            shutil.copy(original_file_path, new_file_path)
+
     def init_new_config(self, file_path: typing.Union[str, pathlib.Path], uproject, engine_dir, p4_settings):
         ''' 
         Initialize new configuration
@@ -2340,16 +2367,31 @@ class Config(object):
 
 
 class UserSettings(object):
-
-    def __init__(self):
+    def init(self):
         try:
             with open(USER_SETTINGS_FILE_PATH) as f:
                 LOGGER.debug(f'Loading Settings {USER_SETTINGS_FILE_PATH}')
                 data = json.load(f)
         except FileNotFoundError:
             # Create a default user_settings
-            data = {}
             LOGGER.debug('Creating default user settings')
+            data = {}
+        except ValueError:
+            LOGGER.error(f'{USER_SETTINGS_FILE_NAME} has invalid JSON format. ')
+            
+            answer = QtWidgets.QMessageBox.question(
+                None,
+                'Invalid User Settings',
+                 f'User settings has invalid JSON format and will be replaced with a valid a new default JSON.'
+                 f'\n\nDo you want to save a backup named {USER_SETTINGS_BACKUP_FILE_NAME} (overrides existing)?',
+                 QtWidgets.QMessageBox.Yes | QtWidgets.QMessageBox.No
+            )
+            if answer == QtWidgets.QMessageBox.Yes:
+                if os.path.exists(USER_SETTINGS_BACKUP_FILE_PATH):
+                    os.remove(USER_SETTINGS_BACKUP_FILE_PATH)
+                shutil.copy(USER_SETTINGS_FILE_PATH, USER_SETTINGS_BACKUP_FILE_PATH)
+                
+            data = {}
 
         self.CONFIG = data.get('config')
         if self.CONFIG:
@@ -2421,11 +2463,12 @@ def list_config_paths() -> typing.List[pathlib.Path]:
     # settings file.
     config_paths = [
         path for path in ROOT_CONFIGS_PATH.rglob(f'*{CONFIG_SUFFIX}')
-        if path != USER_SETTINGS_FILE_PATH]
+        if path != USER_SETTINGS_FILE_PATH and path != USER_SETTINGS_BACKUP_FILE_PATH]
 
     return config_paths
 
 
 # Get the user settings and load their config
 SETTINGS = UserSettings()
-CONFIG = Config(SETTINGS.CONFIG)
+CONFIG = Config()
+    
