Fix bugs related to settings changes:
- Fix not joining multi user section
- Fix engine dir browse requesting file instead of folder
- Fix values from new config dialog not carrying over to project settings
- Fix source control group box not using project setting values

diff --git a/Engine/Plugins/VirtualProduction/Switchboard/Source/Switchboard/switchboard/config.py b/Engine/Plugins/VirtualProduction/Switchboard/Source/Switchboard/switchboard/config.py
--- a/Engine/Plugins/VirtualProduction/Switchboard/Source/Switchboard/switchboard/config.py
+++ b/Engine/Plugins/VirtualProduction/Switchboard/Source/Switchboard/switchboard/config.py
@@ -1890,7 +1890,12 @@ class Config(object):
 
         self.file_path = get_absolute_config_path(file_path)
         self.init_switchboard_settings()
-        self.init_project_settings({ "project_name": self.file_path.stem })
+        self.init_project_settings(
+            { 
+                "project_name": self.file_path.stem, 
+                "uproject": uproject, 
+                "engine_dir": engine_dir,
+            } | p4_settings)
         self.init_unreal_insights()
         self.init_muserver()
 
@@ -1930,7 +1935,7 @@ class Config(object):
                 data.get('uproject', ''),
                 tool_tip="Path to uProject"
             ),
-            "engine_dir": FilePathSetting(
+            "engine_dir": DirectoryPathSetting(
                 "engine_dir",
                 "Engine Directory",
                 data.get('engine_dir', ''),
diff --git a/Engine/Plugins/VirtualProduction/Switchboard/Source/Switchboard/switchboard/settings_dialog.py b/Engine/Plugins/VirtualProduction/Switchboard/Source/Switchboard/switchboard/settings_dialog.py
--- a/Engine/Plugins/VirtualProduction/Switchboard/Source/Switchboard/switchboard/settings_dialog.py
+++ b/Engine/Plugins/VirtualProduction/Switchboard/Source/Switchboard/switchboard/settings_dialog.py
@@ -37,7 +37,10 @@ class SettingsDialog(QtCore.QObject):
         dialog_layout = QtWidgets.QVBoxLayout(self.ui)
         dialog_layout.setContentsMargins(2, 2, 2, 2)
         self.ui.setWindowTitle("Settings")
-        self.ui.finished.connect(self._on_finished)
+        self.ui.finished.connect(
+            lambda result, config=config:
+                self._on_finished(result, config)
+        )
 
         self._create_search_area(dialog_layout)
         self.general_settings_list = [
@@ -99,7 +102,9 @@ class SettingsDialog(QtCore.QObject):
         self._search_settings_again()
         return
 
-    def _on_finished(self, result: int):
+    def _on_finished(self, result: int, config: Config):
+        config.P4_ENABLED.signal_setting_changed.disconnect(self._on_source_control_setting_changed)
+        
         # Currently, the only way to dismiss the settings dialog is by using
         # the close button as opposed to ok/cancel buttons, so we intercept the
         # close to issue a warning if the config path was changed and we're
@@ -193,7 +198,7 @@ class SettingsDialog(QtCore.QObject):
         self._config_paths = config.list_config_paths()
         self.set_config_path(SETTINGS.CONFIG)
         
-        return self.ui.config_path_layout;
+        return self.ui.config_path_layout
         
     def _create_switchboard_settings(self, settings: UserSettings, config: Config):
         self.ui.switchboard_settings_group = CollapsibleGroupBox()
@@ -241,7 +246,16 @@ class SettingsDialog(QtCore.QObject):
         )
         form_layout = QtWidgets.QFormLayout(self.ui.source_control_settings_group)
         layout.addWidget(self.ui.source_control_settings_group)
-        self._create_settings_section(config.source_control_settings, form_layout)
+
+        # Source control setting UI is fully enabled / disabled depending on P4_ENABLED
+        self.ui.source_control_settings_group.setChecked(config.P4_ENABLED.get_value())
+        config.P4_ENABLED.signal_setting_changed.connect(self._on_source_control_setting_changed)
+        settings_to_show = {key: setting for (key, setting) in config.source_control_settings.items() if setting != config.P4_ENABLED}
+        
+        self._create_settings_section(settings_to_show, form_layout)
+        
+    def _on_source_control_setting_changed(self, old_value: bool, new_value: bool):
+        self.ui.source_control_settings_group.setChecked(new_value)
     
     def _create_multi_user_server_settings(self, config: Config):
         self.ui.multi_user_settings = CollapsibleGroupBox()
diff --git a/Engine/Plugins/VirtualProduction/Switchboard/Source/Switchboard/switchboard/switchboard_application.py b/Engine/Plugins/VirtualProduction/Switchboard/Source/Switchboard/switchboard/switchboard_application.py
--- a/Engine/Plugins/VirtualProduction/Switchboard/Source/Switchboard/switchboard/switchboard_application.py
+++ b/Engine/Plugins/VirtualProduction/Switchboard/Source/Switchboard/switchboard/switchboard_application.py
@@ -131,7 +131,7 @@ class MultiUserApplication:
                 return
 
             cmdline = f'start "Multi User Server" "{self.exe_path()}"'
-            cmdline += f' -CONCERTSERVER="{CONFIG.MUSERVER_SERVER_NAME}"'
+            cmdline += f' -CONCERTSERVER="{CONFIG.MUSERVER_SERVER_NAME.get_value()}"'
             cmdline += f' {CONFIG.MUSERVER_COMMAND_LINE_ARGUMENTS.get_value()}'
             cmdline += f' {self.get_mu_server_endpoint_arg()}'
             cmdline += f' {self.get_mu_server_multicast_arg()}'
diff --git a/Engine/Plugins/VirtualProduction/Switchboard/Source/Switchboard/switchboard/switchboard_dialog.py b/Engine/Plugins/VirtualProduction/Switchboard/Source/Switchboard/switchboard/switchboard_dialog.py
--- a/Engine/Plugins/VirtualProduction/Switchboard/Source/Switchboard/switchboard/switchboard_dialog.py
+++ b/Engine/Plugins/VirtualProduction/Switchboard/Source/Switchboard/switchboard/switchboard_dialog.py
@@ -766,8 +766,12 @@ class SwitchboardDialog(QtCore.QObject):
         self.update_current_config_text()
         
     def update_current_config_text(self):
-        file_name = os.path.basename(SETTINGS.CONFIG)
-        self.window.current_config_file_value.setText(file_name)
+        # Can be none when current file is deleted
+        if SETTINGS.CONFIG is not None:
+            file_name = os.path.basename(SETTINGS.CONFIG)
+            self.window.current_config_file_value.setText(file_name)
+        else:
+            self.window.current_config_file_value.setText("No config loaded")
     
     def update_current_ip_address_text(self):
         self.window.current_ip_value.setText(SETTINGS.IP_ADDRESS.get_value())
@@ -809,6 +813,7 @@ class SwitchboardDialog(QtCore.QObject):
             # Update the UI
             self.toggle_p4_controls(CONFIG.P4_ENABLED.get_value())
             self.refresh_levels()
+            self.update_current_config_text()
 
     def menu_delete_config(self):
         """
