Fix not all source control settings being carried over when creating new config

diff --git a/Engine/Plugins/VirtualProduction/Switchboard/Source/Switchboard/switchboard/add_config_dialog.py b/Engine/Plugins/VirtualProduction/Switchboard/Source/Switchboard/switchboard/add_config_dialog.py
--- a/Engine/Plugins/VirtualProduction/Switchboard/Source/Switchboard/switchboard/add_config_dialog.py
+++ b/Engine/Plugins/VirtualProduction/Switchboard/Source/Switchboard/switchboard/add_config_dialog.py
@@ -490,8 +490,8 @@ class AddConfigDialog(QtWidgets.QDialog):
     def p4_settings(self):
         settings = {}
         settings['p4_enabled'] = self.p4_group.isChecked()
-        settings['p4_workspace_name'] = self.p4_workspace_line_edit.text() if self.p4_group.isChecked() else None
-        settings['p4_project_path'] = self.p4_project_path_line_edit.text() if self.p4_group.isChecked() else None
+        settings['source_control_workspace'] = self.p4_workspace_line_edit.text() if self.p4_group.isChecked() else None
+        settings['p4_sync_path'] = self.p4_project_path_line_edit.text() if self.p4_group.isChecked() else None
         settings['p4_engine_path'] = self.p4_engine_path_line_edit.text() if self.p4_group.isChecked() else None
         return settings
 
