Switchboard: Fix dpcvars and execmds for Unreal (not nDisplay) devices.

diff --git a/Engine/Plugins/VirtualProduction/Switchboard/Source/Switchboard/switchboard/devices/unreal/plugin_unreal.py b/Engine/Plugins/VirtualProduction/Switchboard/Source/Switchboard/switchboard/devices/unreal/plugin_unreal.py
--- a/Engine/Plugins/VirtualProduction/Switchboard/Source/Switchboard/switchboard/devices/unreal/plugin_unreal.py
+++ b/Engine/Plugins/VirtualProduction/Switchboard/Source/Switchboard/switchboard/devices/unreal/plugin_unreal.py
@@ -1295,13 +1295,14 @@ class DeviceUnreal(Device):
         exec_cmds = [cmd for cmd in exec_cmds if len(cmd.strip())]
         if len(exec_cmds):
             exec_cmds_expanded = ','.join(exec_cmds)
-            command_line_args.append(f'-ExecCmds="{exec_cmds_expanded}"')
+            command_line_args += f' -ExecCmds="{exec_cmds_expanded}"'
 
-        # DPCVars may need to be appended to, so we don't concatenate them
-        # until the end.
+        # DPCVars may need to be appended to, so we don't concatenate them until the end.
         dp_cvars = DeviceUnreal.csettings["dp_cvars"].get_value(self.name)
+        dp_cvars = [cvar.strip() for cvar in dp_cvars if len(cvar.strip().split('=')) == 2]
 
         (supported_roles, unsupported_roles) = self.get_vproles()
+
         if supported_roles:
             command_line_args += ' -VPRole=' + '|'.join(supported_roles)
         if unsupported_roles:
@@ -1322,14 +1323,13 @@ class DeviceUnreal(Device):
         try:
             if int(max_gpu_count) > 1:
                 command_line_args += f" -MaxGPUCount={max_gpu_count} "
-                if len(dp_cvars):
-                    dp_cvars += ','
-                dp_cvars += 'r.AllowMultiGPUInEditor=1'
+                dp_cvars.append('r.AllowMultiGPUInEditor=1')
         except ValueError:
             LOGGER.warning(f"Invalid Number of GPUs '{max_gpu_count}'")
 
+        # add accumulated dpcvars to args
         if len(dp_cvars):
-            command_line_args += f' -DPCVars="{dp_cvars}" '
+            command_line_args += f" -DPCVars=\"{','.join(dp_cvars)}\""
 
         if DeviceUnreal.csettings['auto_decline_package_recovery'].get_value(
                 self.name):
