Add indication when CL is not built

diff --git a/Engine/Plugins/VirtualProduction/Switchboard/Source/Switchboard/switchboard/devices/device_base.py b/Engine/Plugins/VirtualProduction/Switchboard/Source/Switchboard/switchboard/devices/device_base.py
--- a/Engine/Plugins/VirtualProduction/Switchboard/Source/Switchboard/switchboard/devices/device_base.py
+++ b/Engine/Plugins/VirtualProduction/Switchboard/Source/Switchboard/switchboard/devices/device_base.py
@@ -51,6 +51,7 @@ class DeviceQtHandler(QtCore.QObject):
     signal_device_client_disconnected = QtCore.Signal(object)
     signal_device_project_changelist_changed = QtCore.Signal(object)
     signal_device_engine_changelist_changed = QtCore.Signal(object)
+    signal_device_built_engine_changelist_changed = QtCore.Signal(object)
     signal_device_sync_failed = QtCore.Signal(object)
     signal_device_is_recording_device_changed = QtCore.Signal(object, int)
 
@@ -92,6 +93,7 @@ class Device(QtCore.QObject):
 
         self._project_changelist: Optional[int] = None
         self._engine_changelist: Optional[int] = None
+        self._built_engine_changelist: Optional[int] = None
 
         self._status = DeviceStatus.DISCONNECTED
 
@@ -250,6 +252,17 @@ class Device(QtCore.QObject):
         self._engine_changelist = value
         self.device_qt_handler.signal_device_engine_changelist_changed.emit(
             self)
+        
+    @property
+    def built_engine_changelist(self):
+        return self._built_engine_changelist
+
+    @built_engine_changelist.setter
+    def built_engine_changelist(self, value: Optional[int]):
+        self._built_engine_changelist = value
+        self.device_qt_handler.signal_device_built_engine_changelist_changed.emit(
+            self
+        )
 
     def set_slate(self, value):
         pass
diff --git a/Engine/Plugins/VirtualProduction/Switchboard/Source/Switchboard/switchboard/devices/unreal/plugin_unreal.py b/Engine/Plugins/VirtualProduction/Switchboard/Source/Switchboard/switchboard/devices/unreal/plugin_unreal.py
--- a/Engine/Plugins/VirtualProduction/Switchboard/Source/Switchboard/switchboard/devices/unreal/plugin_unreal.py
+++ b/Engine/Plugins/VirtualProduction/Switchboard/Source/Switchboard/switchboard/devices/unreal/plugin_unreal.py
@@ -3,6 +3,7 @@
 from __future__ import annotations
 
 import base64
+import json
 from collections import OrderedDict
 from datetime import datetime
 from functools import wraps
@@ -37,7 +38,6 @@ from .listener_watcher import ListenerWatcher
 from .redeploy_dialog import RedeployListenerDialog
 from . import version_helpers
 
-
 class ProgramStartQueueItem:
     '''
     Item that holds the information for a program to start in the future.
@@ -791,6 +791,7 @@ class DeviceUnreal(Device):
                 self._request_engine_changelist_number()
         else:
             self.widget.engine_changelist_label.hide()
+            self.widget.update_build_info(current_cl=self.engine_changelist, built_cl=self.built_engine_changelist)
 
     def set_slate(self, value):
         if not self.is_recording_device:
@@ -817,8 +818,19 @@ class DeviceUnreal(Device):
             self.name)
         roles_file_path = os.path.join(
             os.path.dirname(uproject_path), "Config", "Tags", roles_filename)
-        _, msg = message_protocol.create_copy_file_from_listener_message(
-            roles_file_path)
+        _, msg = message_protocol.create_copy_file_from_listener_message(roles_file_path)
+        self.unreal_client.send_message(msg)
+        
+    def _request_unreal_editor_version_file(self):
+        '''
+        Requests the Engine/Binaries/[platform]/UnrealEditor.version file.
+        On success _on_receive_editor_version is called with the file content.
+        '''
+        engine_path = CONFIG.ENGINE_DIR.get_value(self.name).replace('"', '')
+        # TODO: handle different platforms
+        platform_dir = "Win64"
+        roles_file_path = os.path.join(os.path.dirname(engine_path), "Engine", "Binaries", platform_dir, "UnrealEditor.version")
+        _, msg = message_protocol.create_copy_file_from_listener_message(roles_file_path)
         self.unreal_client.send_message(msg)
 
     def _request_project_changelist_number(self):
@@ -1663,22 +1675,34 @@ class DeviceUnreal(Device):
                 f"Error was '{message['error']}''")
 
     def on_file_received(self, source_path, content):
-        if source_path.endswith(
-                DeviceUnreal.csettings["roles_filename"].get_value(self.name)):
-            decoded_content = base64.b64decode(content).decode()
-            tags = parse_unreal_tag_file(decoded_content.splitlines())
-            self.setting_roles.possible_values = tags
-            LOGGER.info(f"{self.name}: All possible roles: {tags}")
-            unsupported_roles = [
-                role for role in self.setting_roles.get_value()
-                if role not in tags]
-            if len(unsupported_roles) > 0:
-                LOGGER.error(
-                    f"{self.name}: Found unsupported roles: "
-                    f"{unsupported_roles}")
-                LOGGER.error(
-                    f"{self.name}: Please change the roles for this device in "
-                    "the settings or in the unreal project settings!")
+        if source_path.endswith(DeviceUnreal.csettings["roles_filename"].get_value(self.name)):
+            self._on_receive_roles(content)
+        elif "UnrealEditor" in source_path and source_path.endswith(".version"):
+            self._on_receive_editor_version(content)
+        
+    def _on_receive_roles(self, content):
+        decoded_content = base64.b64decode(content).decode()
+        tags = parse_unreal_tag_file(decoded_content.splitlines())
+        self.setting_roles.possible_values = tags
+        LOGGER.info(f"{self.name}: All possible roles: {tags}")
+        unsupported_roles = [
+            role for role in self.setting_roles.get_value()
+            if role not in tags]
+        if len(unsupported_roles) > 0:
+            LOGGER.error(
+                f"{self.name}: Found unsupported roles: "
+                f"{unsupported_roles}")
+            LOGGER.error(
+                f"{self.name}: Please change the roles for this device in "
+                "the settings or in the unreal project settings!")
+                
+    def _on_receive_editor_version(self, content):
+        '''
+        Receives the Engine/Binaries/[platform]/UnrealEditor.version file _request_build_version
+        '''
+        decoded_content = base64.b64decode(content).decode()
+        data = json.loads(decoded_content)
+        self.built_engine_changelist = data.get("Changelist", None)
 
     def on_file_receive_failed(self, source_path, error):
         roles = self.setting_roles.get_value()
@@ -1725,7 +1749,9 @@ class DeviceUnreal(Device):
 
                     if '%' == percent[-1]:
                         self.device_qt_handler.signal_device_build_update.emit(
-                            self, step, percent)
+                            self, step, percent
+                        )
+                        
 
         elif process['name'] == 'sync':
             for line in lines:
@@ -1794,6 +1820,7 @@ class DeviceUnreal(Device):
 
         if CONFIG.BUILD_ENGINE.get_value():
             self._request_engine_changelist_number()
+            self._request_unreal_editor_version_file()
 
     def transport_paths(self, device_recording):
         '''
@@ -2026,6 +2053,10 @@ def parse_unreal_tag_file(file_content):
             tags.append(tag)
     return tags
 
+
+BASE_ENGINE_CL_TOOLTIP = "Current Engine Changelist"
+BASE_PROJECT_CL_TOOLTIP = "Current Project Changelist"
+    
 class DeviceWidgetUnreal(DeviceWidget):
 
     signal_open_last_log = QtCore.Signal(object)
@@ -2034,6 +2065,9 @@ class DeviceWidgetUnreal(DeviceWidget):
 
     def __init__(self, name, device_hash, ip_address, icons, parent=None):
         self._autojoin_visible = True
+        self._is_engine_synched = True
+        self._is_project_synched = True
+        self._needs_rebuild = False
 
         super().__init__(name, device_hash, ip_address, icons, parent=parent)
 
@@ -2046,12 +2080,12 @@ class DeviceWidgetUnreal(DeviceWidget):
         changelist_layout = QtWidgets.QVBoxLayout()
         self.engine_changelist_label = QtWidgets.QLabel()
         self.engine_changelist_label.setObjectName('changelist')
-        self.engine_changelist_label.setToolTip("Current Engine Changelist")
+        self.engine_changelist_label.setToolTip(BASE_ENGINE_CL_TOOLTIP)
         changelist_layout.addWidget(self.engine_changelist_label)
 
         self.project_changelist_label = QtWidgets.QLabel()
         self.project_changelist_label.setObjectName('changelist')
-        self.project_changelist_label.setToolTip("Current Project Changelist")
+        self.project_changelist_label.setToolTip(BASE_PROJECT_CL_TOOLTIP)
         changelist_layout.addWidget(self.project_changelist_label)
 
         spacer = QtWidgets.QSpacerItem(
@@ -2159,6 +2193,11 @@ class DeviceWidgetUnreal(DeviceWidget):
         self.sync_button.setDisabled(True)
         self.build_button.setDisabled(True)
 
+        sb_widgets.set_qt_property(self.sync_button, 'not_synched', False)
+        sb_widgets.set_qt_property(self.sync_button, 'not_built', False)
+        sb_widgets.set_qt_property(self.build_button, 'not_synched', False)
+        sb_widgets.set_qt_property(self.build_button, 'not_built', False)
+
     def _update_connected_ui(self):
         ''' Updates the UI of the device to reflect connected status. '''
         # Make sure the button is in the correct state
@@ -2225,13 +2264,16 @@ class DeviceWidgetUnreal(DeviceWidget):
         else:
             self.open_button.setToolTip('Start Unreal')
 
-    def update_project_changelist(self, value):
-        self.project_changelist_label.setText(f'P: {value}')
+    def update_project_changelist(self, required_cl: str, current_device_cl: str):
+        self.project_changelist_label.setText(f'P: {current_device_cl}')
         self.project_changelist_label.setToolTip('Project CL')
 
         self.project_changelist_label.show()
         self.sync_button.show()
         self.build_button.show()
+    
+        is_synched = required_cl is None or required_cl == current_device_cl
+        self.set_project_changelist_is_synched(is_synched)
 
     def update_build_status(self, device, step, percent):
         self.project_changelist_label.setText(f'Building...{percent}')
@@ -2246,18 +2288,72 @@ class DeviceWidgetUnreal(DeviceWidget):
 
         self.project_changelist_label.show()
 
-    def update_engine_changelist(self, value):
-        self.engine_changelist_label.setText(f'E: {value}')
+    def update_engine_changelist(self, required_cl: str, current_device_cl: str, built_device_cl: str):
+        if not CONFIG.BUILD_ENGINE.get_value():
+            return
+        
+        self.engine_changelist_label.setText(f'E: {current_device_cl}')
         self.engine_changelist_label.show()
 
         self.sync_button.show()
         self.build_button.show()
+        
+        is_synched = required_cl is None or required_cl == current_device_cl
+        self.set_engine_changelist_is_synched(is_synched)
+        self.update_build_info(current_cl=current_device_cl, built_cl=built_device_cl)
+
+    def set_project_changelist_is_synched(self, is_synched: bool):
+        self._is_project_synched = is_synched
+        sb_widgets.set_qt_property(self.project_changelist_label, 'not_synched', not is_synched)
+        
+        self._update_cl_widget_tooltip(self.project_changelist_label, BASE_PROJECT_CL_TOOLTIP, self._is_project_synched)
+        self._update_sync_button()
+
+    def set_engine_changelist_is_synched(self, is_synched: bool):
+        self._is_engine_synched = is_synched
+        self._update_engine_cl_label()
+        
+        self._update_cl_widget_tooltip(self.engine_changelist_label, BASE_ENGINE_CL_TOOLTIP, self._is_engine_synched)
+        self._update_sync_button()
+        
+    def _update_sync_button(self):
+        needs_resync = not self._is_project_synched or not self._is_engine_synched
+        sb_widgets.set_qt_property(self.sync_button, 'not_synched', needs_resync)
+            
+    def update_build_info(self, current_cl: str, built_cl: str):
+        if built_cl is not None and current_cl is not None and CONFIG.BUILD_ENGINE.get_value():
+            self._needs_rebuild = int(built_cl) != int(current_cl)
+            self._update_cl_widget_tooltip(self.engine_changelist_label, BASE_ENGINE_CL_TOOLTIP, self._is_engine_synched)
+        else:
+            self._needs_rebuild = False
+
+        tooltip = f"Build changelist.\n\nBuild required.\nCurrent: {current_cl}\nBuilt: {built_cl}" \
+            if self._needs_rebuild else "Build changelist (not required - no changes detected)"
+        self.build_button.setToolTip(tooltip)
+        sb_widgets.set_qt_property(self.build_button, 'not_built', self._needs_rebuild)
+        self._update_engine_cl_label()
+            
+    def _update_engine_cl_label(self):
+        if not self._is_engine_synched:
+            sb_widgets.set_qt_property(self.engine_changelist_label, 'not_synched', True)
+            sb_widgets.set_qt_property(self.engine_changelist_label, 'not_built', False)
+        elif self._needs_rebuild:
+            sb_widgets.set_qt_property(self.engine_changelist_label, 'not_synched', False)
+            sb_widgets.set_qt_property(self.engine_changelist_label, 'not_built', True)
+        else:
+            sb_widgets.set_qt_property(self.engine_changelist_label, 'not_synched', False)
+            sb_widgets.set_qt_property(self.engine_changelist_label, 'not_built', False)
+
+    def _update_cl_widget_tooltip(self, label: QtWidgets.QLabel, base_tooltip: str, is_synched: bool):
+        tooltip = base_tooltip
 
-    def project_changelist_display_warning(self, b):
-        sb_widgets.set_qt_property(self.project_changelist_label, 'error', b)
+        if not is_synched:
+            tooltip += "\nNot synched to selected CL"
+            
+        if self._needs_rebuild:
+            tooltip += "\nSynched CL not built"
 
-    def engine_changelist_display_warning(self, b):
-        sb_widgets.set_qt_property(self.engine_changelist_label, 'error', b)
+        label.setToolTip(tooltip)
 
     def sync_button_clicked(self):
         self.signal_device_widget_sync.emit(self)
diff --git a/Engine/Plugins/VirtualProduction/Switchboard/Source/Switchboard/switchboard/switchboard_dialog.py b/Engine/Plugins/VirtualProduction/Switchboard/Source/Switchboard/switchboard/switchboard_dialog.py
--- a/Engine/Plugins/VirtualProduction/Switchboard/Source/Switchboard/switchboard/switchboard_dialog.py
+++ b/Engine/Plugins/VirtualProduction/Switchboard/Source/Switchboard/switchboard/switchboard_dialog.py
@@ -1015,6 +1015,7 @@ class SwitchboardDialog(QtCore.QObject):
         device.device_qt_handler.signal_device_client_disconnected.connect(self.device_client_disconnected, QtCore.Qt.QueuedConnection)
         device.device_qt_handler.signal_device_project_changelist_changed.connect(self.device_project_changelist_changed, QtCore.Qt.QueuedConnection)
         device.device_qt_handler.signal_device_engine_changelist_changed.connect(self.device_engine_changelist_changed, QtCore.Qt.QueuedConnection)
+        device.device_qt_handler.signal_device_built_engine_changelist_changed.connect(self.device_built_engine_changelist_changed, QtCore.Qt.QueuedConnection)
         device.device_qt_handler.signal_device_status_changed.connect(self.device_status_changed, QtCore.Qt.QueuedConnection)
         device.device_qt_handler.signal_device_sync_failed.connect(self.device_sync_failed, QtCore.Qt.QueuedConnection)
         device.device_qt_handler.signal_device_is_recording_device_changed.connect(self.device_is_recording_device_changed, QtCore.Qt.QueuedConnection)
@@ -1276,17 +1277,12 @@ class SwitchboardDialog(QtCore.QObject):
 
         # Check if all of the devices are on the right changelist
         for device in self.device_manager.devices():
-            if not device.project_changelist:
-                continue
-
-            device_widget = self.device_list_widget.device_widget_by_hash(device.device_hash)
-            if value == EMPTY_SYNC_ENTRY:
-                device_widget.project_changelist_display_warning(False)
-            else:
-                if device.project_changelist == self.project_changelist:
-                    device_widget.project_changelist_display_warning(False)
-                else:
-                    device_widget.project_changelist_display_warning(True)
+            if device.project_changelist:
+                device_widget = self.device_list_widget.device_widget_by_hash(device.device_hash)
+                device_widget.update_project_changelist(
+                    required_cl=self.project_changelist if self.project_changelist != EMPTY_SYNC_ENTRY else None,
+                    current_device_cl=device.project_changelist
+                )
 
     @property
     def engine_changelist(self):
@@ -1296,25 +1292,21 @@ class SwitchboardDialog(QtCore.QObject):
     def engine_changelist(self, value):
         self._set_engine_changelist(value)
 
-    def _set_engine_changelist(self, value):
-        self._engine_changelist = value
+    def _set_engine_changelist(self, changelist_value):
+        self._engine_changelist = changelist_value
 
         if self.window.engine_cl_combo_box.currentText() != self._engine_changelist:
             self.window.engine_cl_combo_box.setText(self._engine_changelist)
 
         # Check if all of the devices are on the right changelist
         for device in self.device_manager.devices():
-            if not device.engine_changelist:
-                continue
-
-            device_widget = self.device_list_widget.device_widget_by_hash(device.device_hash)
-            if value == EMPTY_SYNC_ENTRY:
-                device_widget.engine_changelist_display_warning(False)
-            else:
-                if device.engine_changelist == self.engine_changelist:
-                    device_widget.engine_changelist_display_warning(False)
-                else:
-                    device_widget.engine_changelist_display_warning(True)
+            if device.engine_changelist:
+                device_widget = self.device_list_widget.device_widget_by_hash(device.device_hash)
+                device_widget.update_engine_changelist(
+                    required_cl=self.engine_changelist if self.engine_changelist != EMPTY_SYNC_ENTRY else None, 
+                    current_device_cl=device.engine_changelist,
+                    built_device_cl=device.engine_changelist
+                )
 
     @QtCore.Slot(object)
     def device_widget_connect(self, device_widget):
@@ -1410,28 +1402,42 @@ class SwitchboardDialog(QtCore.QObject):
     @QtCore.Slot(object)
     def device_project_changelist_changed(self, device):
         device_widget = self.device_list_widget.device_widget_by_hash(device.device_hash)
-        device_widget.update_project_changelist(device.project_changelist)
-
-        if self.project_changelist == EMPTY_SYNC_ENTRY:
-            device_widget.project_changelist_display_warning(False)
-        else:
-            if device.project_changelist == self.project_changelist:
-                device_widget.project_changelist_display_warning(False)
-            else:
-                device_widget.project_changelist_display_warning(True)
+        device_widget.update_project_changelist(
+            required_cl=self.project_changelist if self.project_changelist != EMPTY_SYNC_ENTRY else None,
+            current_device_cl=device.project_changelist
+        )
+        
+        cl = device.project_changelist
+        ip = device.ip_address
+        for device in self.device_manager.devices():
+            if device.ip_address == ip and device.project_changelist and device.project_changelist != cl:
+                device.project_changelist = cl
 
     @QtCore.Slot(object)
     def device_engine_changelist_changed(self, device):
         device_widget = self.device_list_widget.device_widget_by_hash(device.device_hash)
-        device_widget.update_engine_changelist(device.engine_changelist)
+        device_widget.update_engine_changelist(
+            required_cl=self.engine_changelist if self.engine_changelist != EMPTY_SYNC_ENTRY else None,
+            current_device_cl=device.engine_changelist,
+            built_device_cl=device.built_engine_changelist
+        )
 
-        if self.engine_changelist == EMPTY_SYNC_ENTRY:
-            device_widget.engine_changelist_display_warning(False)
-        else:
-            if device.engine_changelist == self.engine_changelist:
-                device_widget.engine_changelist_display_warning(False)
-            else:
-                device_widget.engine_changelist_display_warning(True)
+        cl = device.engine_changelist
+        ip = device.ip_address
+        for device in self.device_manager.devices():
+            if device.ip_address == ip and device.engine_changelist and device.engine_changelist != cl:
+                device.engine_changelist = cl
+
+    @QtCore.Slot(object)
+    def device_built_engine_changelist_changed(self, device):
+        device_widget = self.device_list_widget.device_widget_by_hash(device.device_hash)
+        device_widget.update_build_info(current_cl=device.engine_changelist, built_cl=device.built_engine_changelist)
+
+        cl = device.built_engine_changelist
+        ip = device.ip_address
+        for device in self.device_manager.devices():
+            if device.ip_address == ip and device.built_engine_changelist and device.built_engine_changelist != cl:
+                device.built_engine_changelist = cl
 
     @QtCore.Slot(object)
     def device_status_changed(self, device, previous_status):
