Fix wrong argument passed to function

diff --git a/Engine/Plugins/VirtualProduction/Switchboard/Source/Switchboard/switchboard/switchboard_dialog.py b/Engine/Plugins/VirtualProduction/Switchboard/Source/Switchboard/switchboard/switchboard_dialog.py
--- a/Engine/Plugins/VirtualProduction/Switchboard/Source/Switchboard/switchboard/switchboard_dialog.py
+++ b/Engine/Plugins/VirtualProduction/Switchboard/Source/Switchboard/switchboard/switchboard_dialog.py
@@ -1305,7 +1305,7 @@ class SwitchboardDialog(QtCore.QObject):
                 device_widget.update_engine_changelist(
                     required_cl=self.engine_changelist if self.engine_changelist != EMPTY_SYNC_ENTRY else None, 
                     current_device_cl=device.engine_changelist,
-                    built_device_cl=device.engine_changelist
+                    built_device_cl=device.built_engine_changelist
                 )
 
     @QtCore.Slot(object)
