IP address drop-down

diff --git a/Engine/Plugins/VirtualProduction/Switchboard/Source/Switchboard/switchboard/config.py b/Engine/Plugins/VirtualProduction/Switchboard/Source/Switchboard/switchboard/config.py
--- a/Engine/Plugins/VirtualProduction/Switchboard/Source/Switchboard/switchboard/config.py
+++ b/Engine/Plugins/VirtualProduction/Switchboard/Source/Switchboard/switchboard/config.py
@@ -1663,6 +1663,29 @@ class LoggingSetting(Setting):
             widget.model().category_verbosities = new_value
 
 
+class IPAddressSetting(OptionSetting):
+    def __init__(
+        self,
+        attr_name,
+        nice_name,
+        value,
+        tool_tip=None,
+        show_ui=True,
+        allow_reset=True,
+        migrate_data=None
+    ):
+        from socket import getaddrinfo, AF_INET, gethostname
+        ip_addresses = [ip[4][0].__str__() for ip in getaddrinfo(host=gethostname(), port=None, family=AF_INET)]
+        super().__init__(
+            attr_name=attr_name,
+            nice_name=nice_name,
+            value=value,
+            possible_values=ip_addresses,
+            tool_tip=tool_tip,
+            show_ui=show_ui,
+            allow_reset=allow_reset,
+            migrate_data=migrate_data)
+
 class ConfigPathError(Exception):
     '''
     Base exception type for config path related errors.
@@ -2411,7 +2434,7 @@ class UserSettings(object):
             self.CONFIG = config_paths[0] if config_paths else None
 
         # IP Address of the machine running Switchboard
-        self.IP_ADDRESS = StringSetting(
+        self.IP_ADDRESS = IPAddressSetting(
             "ip_address",
             "IP Address",
             data.get("ip_address", socket.gethostbyname(socket.gethostname()))
