Exclude from build

diff --git a/Engine/Plugins/VirtualProduction/Switchboard/Source/Switchboard/switchboard/devices/unreal/plugin_unreal.py b/Engine/Plugins/VirtualProduction/Switchboard/Source/Switchboard/switchboard/devices/unreal/plugin_unreal.py
--- a/Engine/Plugins/VirtualProduction/Switchboard/Source/Switchboard/switchboard/devices/unreal/plugin_unreal.py
+++ b/Engine/Plugins/VirtualProduction/Switchboard/Source/Switchboard/switchboard/devices/unreal/plugin_unreal.py
@@ -509,6 +509,16 @@ class DeviceUnreal(Device):
             value=kwargs.get("last_trace_path",''),
             show_ui=False
         )
+        
+        self.exclude_from_build = BoolSetting(
+            attr_name="exclude_from_build",
+            nice_name="Exclude from build",
+            value=kwargs.get("exclude_from_build", False),
+            tool_tip="Whether to exclude this device from builds"
+        )
+        self.exclude_from_build.signal_setting_changed.connect(
+            lambda _, new_value: self.on_setting_exclude_from_build_changed(new_value)
+        )
 
         self.setting_ip_address.signal_setting_changed.connect(
             self.on_setting_ip_address_changed)
@@ -712,6 +722,7 @@ class DeviceUnreal(Device):
             self.last_launch_command,
             self.last_log_path,
             self.last_trace_path,
+            self.exclude_from_build
         ]
 
     def check_settings_valid(self) -> bool:
@@ -762,6 +773,8 @@ class DeviceUnreal(Device):
 
         super().device_widget_registered(device_widget)
 
+        device_widget.signal_exclude_from_build_toggled.connect(self.on_toggle_exclude_from_build)
+
         # hook to open last log signal from widget
         device_widget.signal_open_last_log.connect(self.on_open_last_log)
 
@@ -793,6 +806,9 @@ class DeviceUnreal(Device):
             self.widget.engine_changelist_label.hide()
             self.widget.update_build_info(current_cl=self.engine_changelist, built_cl=self.built_engine_changelist)
 
+    def on_setting_exclude_from_build_changed(self, exclude_from_build):
+        self.widget.update_exclude_from_build(exclude_from_build, not self.is_disconnected)
+
     def set_slate(self, value):
         if not self.is_recording_device:
             return
@@ -1042,6 +1058,9 @@ class DeviceUnreal(Device):
             self.status = DeviceStatus.SYNCING
 
     def build(self):
+        if self.exclude_from_build.get_value(self.name):
+            return
+        
         program_name = 'build_project'
 
         # check if it is already on its way:
@@ -2023,6 +2042,10 @@ class DeviceUnreal(Device):
         self.autojoin_mu_server.update_value(autojoin.is_autojoin_enabled())
         self.update_settings_menu_state()
 
+    def on_toggle_exclude_from_build(self):
+        self.exclude_from_build.update_value(not self.exclude_from_build.get_value())
+        CONFIG.save()
+
     def on_open_last_log(self):
         ''' Opens the last log in your preferred editor '''
         sb_utils.openfile_with_default_app(self.last_log_path.get_value())
@@ -2058,7 +2081,8 @@ BASE_ENGINE_CL_TOOLTIP = "Current Engine Changelist"
 BASE_PROJECT_CL_TOOLTIP = "Current Project Changelist"
     
 class DeviceWidgetUnreal(DeviceWidget):
-
+    
+    signal_exclude_from_build_toggled = QtCore.Signal()
     signal_open_last_log = QtCore.Signal(object)
     signal_open_last_trace = QtCore.Signal(object)
     signal_copy_last_launch_command = QtCore.Signal(object)
@@ -2068,12 +2092,35 @@ class DeviceWidgetUnreal(DeviceWidget):
         self._is_engine_synched = True
         self._is_project_synched = True
         self._needs_rebuild = False
+        self._exclude_from_build = False
+        self._current_cl = None
+        self._built_cl = None
 
         super().__init__(name, device_hash, ip_address, icons, parent=parent)
 
         CONFIG.P4_ENABLED.signal_setting_changed.connect(
             lambda _, enabled: self.sync_button.setVisible(enabled))
 
+    @property
+    def exclude_from_build(self):
+        return self._exclude_from_build
+
+    def update_exclude_from_build(self, exclude_from_build, update_ui=True):
+        self._exclude_from_build = exclude_from_build
+        if not update_ui:
+            return
+        
+        if exclude_from_build:
+            self.engine_changelist_label.hide()
+            self.build_button.setToolTip("Excluded from build (see device settings)")
+            sb_widgets.set_qt_property(self.build_button, 'not_built', False)
+        else:
+            self.engine_changelist_label.show()
+            self.update_build_info(self._current_cl, self._built_cl)
+
+        self.build_button.setDisabled(exclude_from_build)
+        self._update_engine_cl_label()
+
     def _add_control_buttons(self):
         super()._add_control_buttons()
 
@@ -2206,7 +2253,7 @@ class DeviceWidgetUnreal(DeviceWidget):
 
         self.open_button.setDisabled(False)
         self.sync_button.setDisabled(False)
-        self.build_button.setDisabled(False)
+        self.build_button.setDisabled(self.exclude_from_build)
 
     def update_status(self, status, previous_status):
         super().update_status(status, previous_status)
@@ -2227,7 +2274,9 @@ class DeviceWidgetUnreal(DeviceWidget):
             self.open_button.setDisabled(False)
             self.open_button.setChecked(False)
             self.sync_button.setDisabled(False)
-            self.build_button.setDisabled(False)
+            self.build_button.setDisabled(self.exclude_from_build)
+            if self.exclude_from_build:
+                self.engine_changelist_label.hide()
         elif status == DeviceStatus.CLOSING:
             self.open_button.setDisabled(True)
             self.open_button.setChecked(True)
@@ -2293,7 +2342,8 @@ class DeviceWidgetUnreal(DeviceWidget):
             return
         
         self.engine_changelist_label.setText(f'E: {current_device_cl}')
-        self.engine_changelist_label.show()
+        if not self.exclude_from_build:
+            self.engine_changelist_label.show()
 
         self.sync_button.show()
         self.build_button.show()
@@ -2323,17 +2373,29 @@ class DeviceWidgetUnreal(DeviceWidget):
     def update_build_info(self, current_cl: str, built_cl: str):
         if built_cl is not None and current_cl is not None and CONFIG.BUILD_ENGINE.get_value():
             self._needs_rebuild = int(built_cl) != int(current_cl)
-            self._update_cl_widget_tooltip(self.engine_changelist_label, BASE_ENGINE_CL_TOOLTIP, self._is_engine_synched)
         else:
             self._needs_rebuild = False
 
-        tooltip = f"Build changelist.\n\nBuild required.\nCurrent: {current_cl}\nBuilt: {built_cl}" \
-            if self._needs_rebuild else "Build changelist (not required - no changes detected)"
-        self.build_button.setToolTip(tooltip)
-        sb_widgets.set_qt_property(self.build_button, 'not_built', self._needs_rebuild)
-        self._update_engine_cl_label()
+        self._current_cl = current_cl
+        self._built_cl = built_cl
+
+        should_update_ui = not self.exclude_from_build 
+        if should_update_ui:
+            sb_widgets.set_qt_property(self.build_button, 'not_built', self._needs_rebuild)
+            self._update_build_button_tooltip()
+            self._update_engine_cl_label()
+            self._update_cl_widget_tooltip(self.engine_changelist_label, BASE_ENGINE_CL_TOOLTIP, self._is_engine_synched)
+            
+    def _update_build_button_tooltip(self):
+        if self._built_cl and self._built_cl:
+            tooltip = f"Build changelist.\n\nBuild required.\nCurrent: {self._current_cl}\nBuilt: {self._built_cl}" \
+                if self._needs_rebuild else "Build changelist (not required - no changes detected)"
+            self.build_button.setToolTip(tooltip)
             
     def _update_engine_cl_label(self):
+        if self.exclude_from_build:
+            sb_widgets.set_qt_property(self.engine_changelist_label, 'not_synched', False)
+            sb_widgets.set_qt_property(self.engine_changelist_label, 'not_built', False)
         if not self._is_engine_synched:
             sb_widgets.set_qt_property(self.engine_changelist_label, 'not_synched', True)
             sb_widgets.set_qt_property(self.engine_changelist_label, 'not_built', False)
@@ -2376,6 +2438,10 @@ class DeviceWidgetUnreal(DeviceWidget):
     def populate_context_menu(self, cmenu: QtWidgets.QMenu):
         ''' Called to populate the given context menu with any desired actions'''
         
+        cmenu.addAction(
+            "Include in build" if self.exclude_from_build else "Exclude from build",
+            lambda: self.signal_exclude_from_build_toggled.emit()
+        )
         cmenu.addAction("Open fetched log", lambda: self.signal_open_last_log.emit(self))
         cmenu.addAction("Open fetched trace", lambda: self.signal_open_last_trace.emit(self))
         cmenu.addAction("Copy last launch command", lambda: self.signal_copy_last_launch_command.emit(self))
