Send and process platform name upon initial connection

diff --git a/Engine/Plugins/VirtualProduction/Switchboard/Source/Switchboard/switchboard/devices/unreal/plugin_unreal.py b/Engine/Plugins/VirtualProduction/Switchboard/Source/Switchboard/switchboard/devices/unreal/plugin_unreal.py
--- a/Engine/Plugins/VirtualProduction/Switchboard/Source/Switchboard/switchboard/devices/unreal/plugin_unreal.py
+++ b/Engine/Plugins/VirtualProduction/Switchboard/Source/Switchboard/switchboard/devices/unreal/plugin_unreal.py
@@ -843,8 +843,7 @@ class DeviceUnreal(Device):
         On success _on_receive_editor_version is called with the file content.
         '''
         engine_path = CONFIG.ENGINE_DIR.get_value(self.name).replace('"', '')
-        # TODO: handle different platforms
-        platform_dir = "Win64"
+        platform_dir = self.platform_binary_directory
         roles_file_path = os.path.join(os.path.dirname(engine_path), "Engine", "Binaries", platform_dir, "UnrealEditor.version")
         _, msg = message_protocol.create_copy_file_from_listener_message(roles_file_path)
         self.unreal_client.send_message(msg)
@@ -1823,6 +1822,7 @@ class DeviceUnreal(Device):
         self.os_version_label_sub = message.get('osVersionLabelSub', '')
         self.os_version_number = message.get('osVersionNumber', '')
         self.total_phys_mem = message.get('totalPhysicalMemory', 0)
+        self.platform_binary_directory = message.get('platformBinaryDirectory', '')
 
         # update list of running processes
         for process in message['runningProcesses']:
diff --git a/Engine/Source/Programs/SwitchboardListener/Private/SwitchboardListener.cpp b/Engine/Source/Programs/SwitchboardListener/Private/SwitchboardListener.cpp
--- a/Engine/Source/Programs/SwitchboardListener/Private/SwitchboardListener.cpp
+++ b/Engine/Source/Programs/SwitchboardListener/Private/SwitchboardListener.cpp
@@ -301,6 +301,7 @@ bool FSwitchboardListener::Tick()
 				StatePacket.OsVersionNumber = FPlatformMisc::GetOSVersion();
 
 				StatePacket.TotalPhysicalMemory = FPlatformMemory::GetConstants().TotalPhysical;
+				StatePacket.PlatformBinaryDirectory = FPlatformProcess::GetBinariesSubdirectory();
 
 				for (const auto& RunningProcess : RunningProcesses)
 				{
diff --git a/Engine/Source/Programs/SwitchboardListener/Private/SwitchboardPacket.h b/Engine/Source/Programs/SwitchboardListener/Private/SwitchboardPacket.h
--- a/Engine/Source/Programs/SwitchboardListener/Private/SwitchboardPacket.h
+++ b/Engine/Source/Programs/SwitchboardListener/Private/SwitchboardPacket.h
@@ -81,6 +81,10 @@ struct FSwitchboardStatePacket : public FSwitchboardPacket
 
 	UPROPERTY()
 	uint64 TotalPhysicalMemory;
+
+	/** The name of the platform as it appears in Engine/Binaries */
+	UPROPERTY()
+	FString PlatformBinaryDirectory;
 };
 
 USTRUCT()
