Allow IP address to be editable

diff --git a/Engine/Plugins/VirtualProduction/Switchboard/Source/Switchboard/switchboard/config.py b/Engine/Plugins/VirtualProduction/Switchboard/Source/Switchboard/switchboard/config.py
--- a/Engine/Plugins/VirtualProduction/Switchboard/Source/Switchboard/switchboard/config.py
+++ b/Engine/Plugins/VirtualProduction/Switchboard/Source/Switchboard/switchboard/config.py
@@ -315,6 +315,7 @@ class Setting(QtCore.QObject):
         button.setMaximumWidth(12)
         button.setMaximumHeight(12)
         button.setToolTip("Reset to default")
+        button.setFocusPolicy(QtCore.Qt.FocusPolicy.NoFocus)
         button.pressed.connect(
             lambda override_device_name=override_device_name:
                 self._on_press_reset_override(override_device_name)
@@ -569,6 +570,7 @@ class FileSystemPathSetting(StringSetting):
         edit_layout.addWidget(line_edit)
 
         browse_btn = QtWidgets.QPushButton('Browse')
+        browse_btn.setFocusPolicy(QtCore.Qt.FocusPolicy.NoFocus)
         edit_layout.addWidget(browse_btn)
 
         def on_browse_clicked():
@@ -699,9 +701,7 @@ class OptionSetting(Setting):
 
         self.possible_values = possible_values or []
 
-    def _create_widgets(
-            self, override_device_name: typing.Optional[str] = None) \
-            -> sb_widgets.NonScrollableComboBox:
+    def _create_widgets(self, override_device_name: typing.Optional[str] = None) -> sb_widgets.NonScrollableComboBox:
         combo = sb_widgets.NonScrollableComboBox()
         if self.tool_tip:
             combo.setToolTip(self.tool_tip)
@@ -714,11 +714,10 @@ class OptionSetting(Setting):
 
         self.set_widget(
             widget=combo, override_device_name=override_device_name)
-
-        combo.currentTextChanged.connect(
-            lambda text, override_device_name=override_device_name:
+        combo.currentIndexChanged.connect(
+            lambda index, override_device_name=override_device_name, combo=combo:
                 self._on_widget_value_changed(
-                    text, override_device_name=override_device_name))
+                    combo.itemData(index), override_device_name=override_device_name))
 
         return combo
 
@@ -1676,6 +1675,7 @@ class IPAddressSetting(OptionSetting):
     ):
         from socket import getaddrinfo, AF_INET, gethostname
         ip_addresses = [ip[4][0].__str__() for ip in getaddrinfo(host=gethostname(), port=None, family=AF_INET)]
+        ip_addresses.append("127.0.0.1")
         super().__init__(
             attr_name=attr_name,
             nice_name=nice_name,
@@ -1684,7 +1684,62 @@ class IPAddressSetting(OptionSetting):
             tool_tip=tool_tip,
             show_ui=show_ui,
             allow_reset=allow_reset,
-            migrate_data=migrate_data)
+            migrate_data=migrate_data)#
+        
+        self._finishedEditingFired = False
+
+    def _create_widgets(self, override_device_name: typing.Optional[str] = None) -> sb_widgets.NonScrollableComboBox:
+        combo: sb_widgets.NonScrollableComboBox = super()._create_widgets(override_device_name)
+        combo.setEditable(True)
+        combo.setInsertPolicy(QtWidgets.QComboBox.InsertPolicy.NoInsert)
+        combo.lineEdit().editingFinished.connect(
+            lambda combo=combo: self._validate_and_commit_ip(combo, override_device_name)
+        )
+        
+        return combo
+    
+    def _validate_and_commit_ip(self, combo: sb_widgets.NonScrollableComboBox, override_device_name:str):
+        def is_valid_ip(ip_str:str):
+            try:
+                import ipaddress
+                ip = ipaddress.ip_address(ip_str)
+                return True
+            except:
+                return False
+            
+        def perform_validation():
+            ip_str = combo.lineEdit().text()
+            if not is_valid_ip(ip_str):
+                answer = QtWidgets.QMessageBox.question(
+                    None,
+                    'Invalid IP',
+                    f'The IP \"{ip_str}\" seems to be invalid.\nDo you want to use this IP anyway?',
+                    QtWidgets.QMessageBox.Yes | QtWidgets.QMessageBox.No
+                )
+                if answer == QtWidgets.QMessageBox.No:
+                    combo.blockSignals(True)
+                    combo.lineEdit().blockSignals(True)
+
+                    current_value = self.get_value(override_device_name)
+                    index = combo.findData(current_value)
+                    is_combo_box_entry = index != -1
+                    if is_combo_box_entry:
+                        combo.setCurrentIndex(index)
+                    else:
+                        combo.lineEdit().setText(current_value)
+
+                    combo.lineEdit().blockSignals(False)
+                    combo.blockSignals(False)
+                    return
+
+            self._on_widget_value_changed(ip_str, override_device_name=override_device_name)
+        
+        if self._finishedEditingFired:
+            return
+        else:
+            self._finishedEditingFired = True
+            perform_validation()
+            self._finishedEditingFired = False
 
 class ConfigPathError(Exception):
     '''
