Fix "exlcude from build" UI not being updated after launch

diff --git a/Engine/Plugins/VirtualProduction/Switchboard/Source/Switchboard/switchboard/devices/unreal/plugin_unreal.py b/Engine/Plugins/VirtualProduction/Switchboard/Source/Switchboard/switchboard/devices/unreal/plugin_unreal.py
--- a/Engine/Plugins/VirtualProduction/Switchboard/Source/Switchboard/switchboard/devices/unreal/plugin_unreal.py
+++ b/Engine/Plugins/VirtualProduction/Switchboard/Source/Switchboard/switchboard/devices/unreal/plugin_unreal.py
@@ -516,9 +516,6 @@ class DeviceUnreal(Device):
             value=kwargs.get("exclude_from_build", False),
             tool_tip="Whether to exclude this device from builds"
         )
-        self.exclude_from_build.signal_setting_changed.connect(
-            lambda _, new_value: self.on_setting_exclude_from_build_changed(new_value)
-        )
 
         self.setting_ip_address.signal_setting_changed.connect(
             self.on_setting_ip_address_changed)
@@ -618,6 +615,15 @@ class DeviceUnreal(Device):
         # Notify user of any invalid settings
         self.check_settings_valid()
 
+    def init(self, widget_class, icons):
+        super().init(widget_class, icons)
+        
+        self.exclude_from_build.signal_setting_changed.connect(
+            lambda _, new_value: self.on_setting_exclude_from_build_changed(new_value)
+        )
+        self.on_setting_exclude_from_build_changed(self.exclude_from_build.get_value())
+        
+
     def should_allow_exit(self, close_req_id: int) -> bool:
         # Delegate to a class method which surveys all active devices.
         return DeviceUnreal._should_allow_exit(close_req_id)
@@ -1057,7 +1063,7 @@ class DeviceUnreal(Device):
             self.status = DeviceStatus.SYNCING
 
     def build(self):
-        if self.exclude_from_build.get_value(self.name):
+        if self.exclude_from_build.get_value():
             return
         
         program_name = 'build_project'
