Change how build button is updated

diff --git a/Engine/Plugins/VirtualProduction/Switchboard/Source/Switchboard/switchboard/devices/unreal/plugin_unreal.py b/Engine/Plugins/VirtualProduction/Switchboard/Source/Switchboard/switchboard/devices/unreal/plugin_unreal.py
--- a/Engine/Plugins/VirtualProduction/Switchboard/Source/Switchboard/switchboard/devices/unreal/plugin_unreal.py
+++ b/Engine/Plugins/VirtualProduction/Switchboard/Source/Switchboard/switchboard/devices/unreal/plugin_unreal.py
@@ -1639,11 +1639,13 @@ class DeviceUnreal(Device):
 
             if 'build_project' == program_name:
                 self.status = DeviceStatus.CLOSED
+                
+                self._request_project_changelist_number()
                 # Forces an update to the changelist field (to hide the
                 # Building state).
-                self.project_changelist = self.project_changelist
                 if CONFIG.BUILD_ENGINE.get_value():
-                    self.engine_changelist = self.engine_changelist
+                    self._request_engine_changelist_number()
+                    self._request_unreal_editor_version_file()
 
         elif "cstat" in program_name:
             output = get_stdout_str()
@@ -2099,8 +2101,7 @@ class DeviceWidgetUnreal(DeviceWidget):
         self._is_project_synched = True
         self._needs_rebuild = False
         self._exclude_from_build = False
-        self._current_cl = None
-        self._built_cl = None
+        self._desired_build_button_tooltip = "Build changelist"
 
         super().__init__(name, device_hash, ip_address, icons, parent=parent)
 
@@ -2118,11 +2119,11 @@ class DeviceWidgetUnreal(DeviceWidget):
         
         if exclude_from_build:
             self.engine_changelist_label.hide()
-            self.build_button.setToolTip("Excluded from build (see device settings)")
+            self._update_build_button_tooltip()
             sb_widgets.set_qt_property(self.build_button, 'not_built', False)
         else:
             self.engine_changelist_label.show()
-            self.update_build_info(self._current_cl, self._built_cl)
+            self._refresh_build_info_ui()
 
         self.build_button.setDisabled(exclude_from_build)
         self._update_engine_cl_label()
@@ -2377,15 +2378,22 @@ class DeviceWidgetUnreal(DeviceWidget):
         sb_widgets.set_qt_property(self.sync_button, 'not_synched', needs_resync)
             
     def update_build_info(self, current_cl: str, built_cl: str):
+        def set_desired_tooltip(current_cl: str, built_cl: str):
+            desired_tooltip = f"Build changelist.\n\nBuild required.\nCurrent: {current_cl}\nBuilt: {built_cl}" \
+                if self._needs_rebuild else "Build changelist (not required - no changes detected)"
+            self._desired_build_button_tooltip = desired_tooltip
+            self._update_build_button_tooltip()
+        
         if built_cl is not None and current_cl is not None and CONFIG.BUILD_ENGINE.get_value():
             self._needs_rebuild = int(built_cl) != int(current_cl)
         else:
             self._needs_rebuild = False
 
-        self._current_cl = current_cl
-        self._built_cl = built_cl
-
-        should_update_ui = not self.exclude_from_build 
+        set_desired_tooltip(current_cl, built_cl)
+        self._refresh_build_info_ui()
+        
+    def _refresh_build_info_ui(self):
+        should_update_ui = not self.exclude_from_build
         if should_update_ui:
             sb_widgets.set_qt_property(self.build_button, 'not_built', self._needs_rebuild)
             self._update_build_button_tooltip()
@@ -2393,10 +2401,10 @@ class DeviceWidgetUnreal(DeviceWidget):
             self._update_cl_widget_tooltip(self.engine_changelist_label, BASE_ENGINE_CL_TOOLTIP, self._is_engine_synched)
             
     def _update_build_button_tooltip(self):
-        if self._built_cl and self._built_cl:
-            tooltip = f"Build changelist.\n\nBuild required.\nCurrent: {self._current_cl}\nBuilt: {self._built_cl}" \
-                if self._needs_rebuild else "Build changelist (not required - no changes detected)"
-            self.build_button.setToolTip(tooltip)
+        if self.exclude_from_build:
+            self.build_button.setToolTip("Excluded from build (see device settings)")
+        else:
+            self.build_button.setToolTip(self._desired_build_button_tooltip)
             
     def _update_engine_cl_label(self):
         if self.exclude_from_build:
