Switchboard: Limit substitution of "Content" -> "Game" in map path to first instance.

diff --git a/Engine/Plugins/VirtualProduction/Switchboard/Source/Switchboard/switchboard/config.py b/Engine/Plugins/VirtualProduction/Switchboard/Source/Switchboard/switchboard/config.py
--- a/Engine/Plugins/VirtualProduction/Switchboard/Source/Switchboard/switchboard/config.py
+++ b/Engine/Plugins/VirtualProduction/Switchboard/Source/Switchboard/switchboard/config.py
@@ -2382,7 +2382,7 @@ class Config(object):
                 map_name, _ = os.path.splitext(name)
                 path_name = path_to_map.replace(project_dir, '', 1)
                 # Ignore the fact that maps may exist in plugins - we only search game content
-                path_name = path_name.replace('Content', 'Game')
+                path_name = path_name.replace('Content', 'Game', 1)
                 path_name = os.path.join(path_name, map_name)
                 path_name = path_name.replace(os.sep, '/')
                 
