Switchboard: Improve printout of exception related to auto-filling p4 settings when creating a new config.

#skipci

diff --git a/Engine/Plugins/VirtualProduction/Switchboard/Source/Switchboard/switchboard/add_config_dialog.py b/Engine/Plugins/VirtualProduction/Switchboard/Source/Switchboard/switchboard/add_config_dialog.py
--- a/Engine/Plugins/VirtualProduction/Switchboard/Source/Switchboard/switchboard/add_config_dialog.py
+++ b/Engine/Plugins/VirtualProduction/Switchboard/Source/Switchboard/switchboard/add_config_dialog.py
@@ -295,7 +295,7 @@ class AddConfigDialog(QtWidgets.QDialog):
             ]
 
         except Exception as e:
-            LOGGER.warning(f'Could not query p4: {e}')
+            LOGGER.warning(f'Could not query p4: {repr(e)}')
             return
 
         client = diagState.p4Workspace
@@ -309,7 +309,7 @@ class AddConfigDialog(QtWidgets.QDialog):
                 diagState.p4Workspace = client
                 diagState.p4Engine = p4path
             except Exception as e:
-                LOGGER.warning(f"Could not auto-fill Engine p4 settings: {e}")
+                LOGGER.warning(f"Could not auto-fill Engine p4 settings: {repr(e)}")
 
         # Infer project p4 path
         if bUpdateProject:
@@ -321,7 +321,7 @@ class AddConfigDialog(QtWidgets.QDialog):
                 diagState.p4Workspace = client
                 diagState.p4Project = p4path
             except Exception as e:
-                LOGGER.warning(f"Could not auto-fill project p4 settings: {e}")
+                LOGGER.warning(f"Could not auto-fill project p4 settings: {repr(e)}")
 
 
     def on_p4_toggled(self, checked):
diff --git a/Engine/Plugins/VirtualProduction/Switchboard/Source/Switchboard/switchboard/p4_utils.py b/Engine/Plugins/VirtualProduction/Switchboard/Source/Switchboard/switchboard/p4_utils.py
--- a/Engine/Plugins/VirtualProduction/Switchboard/Source/Switchboard/switchboard/p4_utils.py
+++ b/Engine/Plugins/VirtualProduction/Switchboard/Source/Switchboard/switchboard/p4_utils.py
@@ -15,7 +15,7 @@ def p4_login(f):
         try:
             return f(*args, **kwargs)
         except Exception as e:
-            LOGGER.error(f'{e}')
+            LOGGER.error(f'{repr(e)}')
             LOGGER.error('Error running P4 command. Please make sure you are logged into Perforce and environment variables are set')
             return None
 
