Fix crash when connecting to device

diff --git a/Engine/Plugins/VirtualProduction/Switchboard/Source/Switchboard/switchboard/devices/unreal/plugin_unreal.py b/Engine/Plugins/VirtualProduction/Switchboard/Source/Switchboard/switchboard/devices/unreal/plugin_unreal.py
--- a/Engine/Plugins/VirtualProduction/Switchboard/Source/Switchboard/switchboard/devices/unreal/plugin_unreal.py
+++ b/Engine/Plugins/VirtualProduction/Switchboard/Source/Switchboard/switchboard/devices/unreal/plugin_unreal.py
@@ -1724,7 +1724,7 @@ class DeviceUnreal(Device):
                 
     def _on_receive_editor_version(self, content):
         '''
-        Receives the Engine/Binaries/[platform]/UnrealEditor.version file _request_build_version
+        Receives the Engine/Binaries/[platform]/UnrealEditor.version file _request_unreal_editor_version_file
         '''
         decoded_content = base64.b64decode(content).decode()
         data = json.loads(decoded_content)
@@ -2329,7 +2329,7 @@ class DeviceWidgetUnreal(DeviceWidget):
         self.build_button.show()
     
         is_synched = required_cl is None or required_cl == current_device_cl
-        self.set_project_changelist_is_synched(is_synched)
+        self._set_project_changelist_is_synched(is_synched)
 
     def update_build_status(self, device, step, percent):
         self.project_changelist_label.setText(f'Building...{percent}')
@@ -2356,17 +2356,17 @@ class DeviceWidgetUnreal(DeviceWidget):
         self.build_button.show()
         
         is_synched = required_cl is None or required_cl == current_device_cl
-        self.set_engine_changelist_is_synched(is_synched)
+        self._set_engine_changelist_is_synched(is_synched)
         self.update_build_info(current_cl=current_device_cl, built_cl=built_device_cl)
 
-    def set_project_changelist_is_synched(self, is_synched: bool):
+    def _set_project_changelist_is_synched(self, is_synched: bool):
         self._is_project_synched = is_synched
         sb_widgets.set_qt_property(self.project_changelist_label, 'not_synched', not is_synched)
         
         self._update_cl_widget_tooltip(self.project_changelist_label, BASE_PROJECT_CL_TOOLTIP, self._is_project_synched)
         self._update_sync_button()
 
-    def set_engine_changelist_is_synched(self, is_synched: bool):
+    def _set_engine_changelist_is_synched(self, is_synched: bool):
         self._is_engine_synched = is_synched
         self._update_engine_cl_label()
         
@@ -2380,7 +2380,7 @@ class DeviceWidgetUnreal(DeviceWidget):
     def update_build_info(self, current_cl: str, built_cl: str):
         def set_desired_tooltip(current_cl: str, built_cl: str):
             desired_tooltip = f"Build changelist.\n\nBuild required.\nCurrent: {current_cl}\nBuilt: {built_cl}" \
-                if self._needs_rebuild else "Build changelist (not required - no changes detected)"
+                if self._needs_rebuild else f"Build changelist (not required - built CL {built_cl} matches synched CL)"
             self._desired_build_button_tooltip = desired_tooltip
             self._update_build_button_tooltip()
         
diff --git a/Engine/Plugins/VirtualProduction/Switchboard/Source/Switchboard/switchboard/switchboard_dialog.py b/Engine/Plugins/VirtualProduction/Switchboard/Source/Switchboard/switchboard/switchboard_dialog.py
--- a/Engine/Plugins/VirtualProduction/Switchboard/Source/Switchboard/switchboard/switchboard_dialog.py
+++ b/Engine/Plugins/VirtualProduction/Switchboard/Source/Switchboard/switchboard/switchboard_dialog.py
@@ -1433,12 +1433,6 @@ class SwitchboardDialog(QtCore.QObject):
         device_widget = self.device_list_widget.device_widget_by_hash(device.device_hash)
         device_widget.update_build_info(current_cl=device.engine_changelist, built_cl=device.built_engine_changelist)
 
-        cl = device.built_engine_changelist
-        ip = device.ip_address
-        for device in self.device_manager.devices():
-            if device.ip_address == ip and device.built_engine_changelist and device.built_engine_changelist != cl:
-                device.built_engine_changelist = cl
-
     @QtCore.Slot(object)
     def device_status_changed(self, device, previous_status):
         # Update the device widget
