Zip log files

diff --git a/Engine/Plugins/VirtualProduction/Switchboard/Source/Switchboard/switchboard/switchboard_dialog.py b/Engine/Plugins/VirtualProduction/Switchboard/Source/Switchboard/switchboard/switchboard_dialog.py
--- a/Engine/Plugins/VirtualProduction/Switchboard/Source/Switchboard/switchboard/switchboard_dialog.py
+++ b/Engine/Plugins/VirtualProduction/Switchboard/Source/Switchboard/switchboard/switchboard_dialog.py
@@ -252,6 +252,7 @@ class SwitchboardDialog(QtCore.QObject):
 
         # Convenience Open Logs Folder menu item
         self.register_open_logs_menuitem()
+        self.register_zip_logs_menuitem()
 
         # Transport Manager
         #self.transport_queue = recording.TransportQueue(CONFIG.SWITCHBOARD_DIR)
@@ -535,12 +536,56 @@ class SwitchboardDialog(QtCore.QObject):
 
     def register_open_logs_menuitem(self):
         ''' Registers convenience "Open Logs Folder" menu item '''
+        action = self.register_tools_menu_action("&Open Logs Folder")
+        action.triggered.connect(self._open_logs_folder)
 
+    def _open_logs_folder(self):
+        sb_utils.explore_path(DeviceUnreal.get_log_download_dir())
+     
+    def register_zip_logs_menuitem(self):
         def open_logs_folder():
-            sb_utils.explore_path(DeviceUnreal.get_log_download_dir())
-
-        action = self.register_tools_menu_action("&Open Logs Folder")
+            # Skip back up files
+            unreal_devices = [device for device in self.device_manager.devices() if isinstance(device, DeviceUnreal)]
+            latest_log_file_names = list(map(lambda device: device.last_log_path.get_value(), unreal_devices))
+            is_latest_log_file = lambda file_path: file_path in latest_log_file_names
+
+            running_devices = [device for device in unreal_devices if device.status in [DeviceStatus.CONNECTING, DeviceStatus.OPEN, DeviceStatus.CLOSING]] 
+            if len (running_devices) > 0:
+                pretty_devices_str = ""
+                for name in list(map(lambda device: device.name, running_devices)):
+                    pretty_devices_str += f" - {name}\n"
+                continue_answer = QtWidgets.QMessageBox.question(
+                    None,
+                    "Devices still running",
+                    f"The following devices are still running:\n{ pretty_devices_str }"
+                    "\nThese devices should be closed to have their running logs included.\nDo you want to continue?",
+                    QtWidgets.QMessageBox.Yes | QtWidgets.QMessageBox.No
+                )
+                if continue_answer == QtWidgets.QMessageBox.No:
+                    return
+
+            folder_to_zip = DeviceUnreal.get_log_download_dir()
+            zip_destination = DeviceUnreal.get_log_download_dir() / "Logs.zip"
+            if os.path.exists(zip_destination):
+                should_replace_answer = QtWidgets.QMessageBox.question(
+                    None,
+                    "Replace existing",
+                    "The file \"Logs.zip\" already exists. Do you want to replace it?",
+                    QtWidgets.QMessageBox.Yes | QtWidgets.QMessageBox.No
+                )
+                if should_replace_answer == QtWidgets.QMessageBox.No:
+                    return
+                
+            sb_utils.zip_files_in_folder(
+                zipped_directory=folder_to_zip,
+                zip_result_path=zip_destination, 
+                allow_file=is_latest_log_file
+            )
+            self._open_logs_folder()
+
+        action = self.register_tools_menu_action("&Zip Logs")
         action.triggered.connect(open_logs_folder)
+        
 
     def add_tools_menu(self):
         ''' Adds tools menu to menu bar and populates built-in items '''
diff --git a/Engine/Plugins/VirtualProduction/Switchboard/Source/Switchboard/switchboard/switchboard_utils.py b/Engine/Plugins/VirtualProduction/Switchboard/Source/Switchboard/switchboard/switchboard_utils.py
--- a/Engine/Plugins/VirtualProduction/Switchboard/Source/Switchboard/switchboard/switchboard_utils.py
+++ b/Engine/Plugins/VirtualProduction/Switchboard/Source/Switchboard/switchboard/switchboard_utils.py
@@ -1,9 +1,14 @@
 # Copyright Epic Games, Inc. All Rights Reserved.
+import typing
 
 from .switchboard_logging import LOGGER
 
 from enum import Enum
 
+from os import listdir
+from os.path import isfile, join
+from zipfile import ZipFile
+
 import datetime
 import subprocess
 import os
@@ -185,6 +190,17 @@ def explore_path(path:str):
     else:
         LOGGER.error(f"explore_path not supported in platform '{sys.platform}'")
 
+def zip_files_in_folder(zipped_directory: str, zip_result_path: str, allow_file: typing.Callable[[str], bool]):
+    ''' Zips all files within zipped_directory and saves it as zip_result_path. Recursively iterates sub-folders.'''
+    with ZipFile(zip_result_path, "w") as zip_file:
+        for (dirpath, dirnames, filenames) in os.walk(zipped_directory):
+            for file_name in filenames:
+                file_path = os.path.join(dirpath, file_name)
+                if allow_file(file_path):
+                    directory_name_in_zip = pathlib.Path(zipped_directory).relative_to(pathlib.Path(dirpath))
+                    name_in_zip = os.path.join(directory_name_in_zip, file_name) if directory_name_in_zip else file_name 
+                    zip_file.write(file_path, name_in_zip)
+
 def copy2clipboard(text:str):
     if sys.platform.startswith('win'):
         subprocess.Popen(f'echo[{text.strip()}|clip', shell=True)
