Fix crash on right-click

diff --git a/Engine/Plugins/VirtualProduction/Switchboard/Source/Switchboard/switchboard/device_list_widget.py b/Engine/Plugins/VirtualProduction/Switchboard/Source/Switchboard/switchboard/device_list_widget.py
--- a/Engine/Plugins/VirtualProduction/Switchboard/Source/Switchboard/switchboard/device_list_widget.py
+++ b/Engine/Plugins/VirtualProduction/Switchboard/Source/Switchboard/switchboard/device_list_widget.py
@@ -50,6 +50,10 @@ class DeviceListWidget(QtWidgets.QListWidget):
 
         item_widget = self.itemWidget(item)
         assert item_widget
+        
+        # No actions for clicking the header widget
+        if isinstance(item_widget, DeviceWidgetHeader):
+            return
 
         device_context_menu = QtWidgets.QMenu(self)
 
