Include utrace files in exported logs and timestamp zip

diff --git a/Engine/Plugins/VirtualProduction/Switchboard/Source/Switchboard/switchboard/switchboard_dialog.py b/Engine/Plugins/VirtualProduction/Switchboard/Source/Switchboard/switchboard/switchboard_dialog.py
--- a/Engine/Plugins/VirtualProduction/Switchboard/Source/Switchboard/switchboard/switchboard_dialog.py
+++ b/Engine/Plugins/VirtualProduction/Switchboard/Source/Switchboard/switchboard/switchboard_dialog.py
@@ -547,7 +547,9 @@ class SwitchboardDialog(QtCore.QObject):
             # Skip back up files
             unreal_devices = [device for device in self.device_manager.devices() if isinstance(device, DeviceUnreal)]
             latest_log_file_names = list(map(lambda device: device.last_log_path.get_value(), unreal_devices))
-            is_latest_log_file = lambda file_path: file_path in latest_log_file_names
+            last_trace_path = list(map(lambda device: device.last_trace_path.get_value(), unreal_devices))
+            is_latest_log_or_trace_file = \
+                lambda file_path: file_path in latest_log_file_names or file_path in last_trace_path
 
             running_devices = [device for device in unreal_devices if device.status in [DeviceStatus.CONNECTING, DeviceStatus.OPEN, DeviceStatus.CLOSING]] 
             if len (running_devices) > 0:
@@ -565,21 +567,12 @@ class SwitchboardDialog(QtCore.QObject):
                     return
 
             folder_to_zip = DeviceUnreal.get_log_download_dir()
-            zip_destination = DeviceUnreal.get_log_download_dir() / "Logs.zip"
-            if os.path.exists(zip_destination):
-                should_replace_answer = QtWidgets.QMessageBox.question(
-                    None,
-                    "Replace existing",
-                    "The file \"Logs.zip\" already exists. Do you want to replace it?",
-                    QtWidgets.QMessageBox.Yes | QtWidgets.QMessageBox.No
-                )
-                if should_replace_answer == QtWidgets.QMessageBox.No:
-                    return
-                
+            modtime = datetime.datetime.now()
+            zip_destination = DeviceUnreal.get_log_download_dir() / f"Logs-{modtime.strftime('%Y.%m.%d-%H.%M.%S')}.zip"
             sb_utils.zip_files_in_folder(
                 zipped_directory=folder_to_zip,
                 zip_result_path=zip_destination, 
-                allow_file=is_latest_log_file
+                allow_file=is_latest_log_or_trace_file
             )
             self._open_logs_folder()
 
