Zip logs improvements:
- give custom zip name
- include multiuser logs

diff --git a/Engine/Plugins/VirtualProduction/Switchboard/Source/Switchboard/switchboard/config.py b/Engine/Plugins/VirtualProduction/Switchboard/Source/Switchboard/switchboard/config.py
--- a/Engine/Plugins/VirtualProduction/Switchboard/Source/Switchboard/switchboard/config.py
+++ b/Engine/Plugins/VirtualProduction/Switchboard/Source/Switchboard/switchboard/config.py
@@ -2396,6 +2396,12 @@ class Config(object):
         return self.engine_exe_path(
             self.ENGINE_DIR.get_value(), self.MULTIUSER_SERVER_EXE.get_value())
 
+    def multiuser_server_session_directory_path(self):
+        return os.path.join(self.ENGINE_DIR.get_value(), "Programs", "UnrealMultiUserServer", "Intermediate", "MultiUser")
+        
+    def multiuser_server_log_path(self):
+        return os.path.join(self.ENGINE_DIR.get_value(), "Programs", "UnrealMultiUserServer", "Saved", "Logs", "UnrealMultiUserServer.log")
+
     def listener_path(self):
         return self.engine_exe_path(
             self.ENGINE_DIR.get_value(), self.LISTENER_EXE.get_value())
diff --git a/Engine/Plugins/VirtualProduction/Switchboard/Source/Switchboard/switchboard/switchboard_dialog.py b/Engine/Plugins/VirtualProduction/Switchboard/Source/Switchboard/switchboard/switchboard_dialog.py
--- a/Engine/Plugins/VirtualProduction/Switchboard/Source/Switchboard/switchboard/switchboard_dialog.py
+++ b/Engine/Plugins/VirtualProduction/Switchboard/Source/Switchboard/switchboard/switchboard_dialog.py
@@ -33,8 +33,7 @@ from switchboard.switchboard_logging import ConsoleStream, LOGGER
 from switchboard.tools.insights_launcher import InsightsLauncher
 from switchboard.tools.listener_launcher import ListenerLauncher
 from switchboard.devices.unreal.plugin_unreal import DeviceUnreal
-
-from . import switchboard_utils as sb_utils
+from switchboard.util import collect_logs
 
 ENGINE_PATH = "../../../../.."
 RELATIVE_PATH = os.path.dirname(__file__)
@@ -537,48 +536,17 @@ class SwitchboardDialog(QtCore.QObject):
     def register_open_logs_menuitem(self):
         ''' Registers convenience "Open Logs Folder" menu item '''
         action = self.register_tools_menu_action("&Open Logs Folder")
-        action.triggered.connect(self._open_logs_folder)
-
-    def _open_logs_folder(self):
-        sb_utils.explore_path(DeviceUnreal.get_log_download_dir())
+        action.triggered.connect(collect_logs.open_logs_folder)
      
     def register_zip_logs_menuitem(self):
-        def open_logs_folder():
-            # Skip back up files
-            unreal_devices = [device for device in self.device_manager.devices() if isinstance(device, DeviceUnreal)]
-            latest_log_file_names = list(map(lambda device: device.last_log_path.get_value(), unreal_devices))
-            last_trace_path = list(map(lambda device: device.last_trace_path.get_value(), unreal_devices))
-            is_latest_log_or_trace_file = \
-                lambda file_path: file_path in latest_log_file_names or file_path in last_trace_path
-
-            running_devices = [device for device in unreal_devices if device.status in [DeviceStatus.CONNECTING, DeviceStatus.OPEN, DeviceStatus.CLOSING]] 
-            if len (running_devices) > 0:
-                pretty_devices_str = ""
-                for name in list(map(lambda device: device.name, running_devices)):
-                    pretty_devices_str += f" - {name}\n"
-                continue_answer = QtWidgets.QMessageBox.question(
-                    None,
-                    "Devices still running",
-                    f"The following devices are still running:\n{ pretty_devices_str }"
-                    "\nThese devices should be closed to have their running logs included.\nDo you want to continue?",
-                    QtWidgets.QMessageBox.Yes | QtWidgets.QMessageBox.No
-                )
-                if continue_answer == QtWidgets.QMessageBox.No:
-                    return
-
-            folder_to_zip = DeviceUnreal.get_log_download_dir()
-            modtime = datetime.datetime.now()
-            zip_destination = DeviceUnreal.get_log_download_dir() / f"Logs-{modtime.strftime('%Y.%m.%d-%H.%M.%S')}.zip"
-            sb_utils.zip_files_in_folder(
-                zipped_directory=folder_to_zip,
-                zip_result_path=zip_destination, 
-                allow_file=is_latest_log_or_trace_file
+        def save_logs():
+            collect_logs.execute_zip_logs_workflow(
+                CONFIG,
+                [device for device in self.device_manager.devices() if isinstance(device, DeviceUnreal)]
             )
-            self._open_logs_folder()
-
+            collect_logs.open_logs_folder()
         action = self.register_tools_menu_action("&Zip Logs")
-        action.triggered.connect(open_logs_folder)
-        
+        action.triggered.connect(save_logs)
 
     def add_tools_menu(self):
         ''' Adds tools menu to menu bar and populates built-in items '''
diff --git a/Engine/Plugins/VirtualProduction/Switchboard/Source/Switchboard/switchboard/switchboard_logging.py b/Engine/Plugins/VirtualProduction/Switchboard/Source/Switchboard/switchboard/switchboard_logging.py
--- a/Engine/Plugins/VirtualProduction/Switchboard/Source/Switchboard/switchboard/switchboard_logging.py
+++ b/Engine/Plugins/VirtualProduction/Switchboard/Source/Switchboard/switchboard/switchboard_logging.py
@@ -97,6 +97,7 @@ class HTMLogger(object):
         super(HTMLogger, self).__init__()
         self.logger = logger
         self.file_handler = None
+        self.log_path = None
 
         self.add_custom_levels()
 
@@ -165,8 +166,10 @@ class HTMLogger(object):
         if not log_path:
             if not os.path.isdir(DEFAULT_LOG_PATH):
                 os.makedirs(DEFAULT_LOG_PATH)
-            log_path = os.path.join(DEFAULT_LOG_PATH, "switchboard_{}.html".format(str(time.time())))
-        self.file_handler = logging.FileHandler(log_path)
+            self.log_path = os.path.join(DEFAULT_LOG_PATH, "switchboard_{}.html".format(str(time.time())))
+        else:
+            self.log_path = log_path
+        self.file_handler = logging.FileHandler(self.log_path)
         self.logger.addHandler(self.file_handler)
         self.logger.setLevel(logging.DEBUG)  # Log everything.
 
@@ -174,12 +177,22 @@ class HTMLogger(object):
         self.logger.setLevel(value)
 
     def disable_file_logging(self):
+        self.log_path = None
         if not self.file_handler:
             return self.logger.warning("No file handler found!")
         self.logger.removeHandler(self.file_handler)
         self.file_handler.close()
-
-
+        
+    def save_log_file(self) -> str:
+        if self.log_path is not None:
+            replaced_path_name = self.log_path.replace(".html", ".log")
+            with open(replaced_path_name, "w") as file:
+                file.write(_LOG_STREAM.getvalue())
+            return replaced_path_name
+
+from io import StringIO
+_LOG_STREAM = StringIO()
+logging.basicConfig(stream=_LOG_STREAM)
 _LOGGER = logging.getLogger("switchboard")
 DEFAULT_LOG_PATH = os.path.join(tempfile.gettempdir(), "switchboard")
 
diff --git a/Engine/Plugins/VirtualProduction/Switchboard/Source/Switchboard/switchboard/switchboard_utils.py b/Engine/Plugins/VirtualProduction/Switchboard/Source/Switchboard/switchboard/switchboard_utils.py
--- a/Engine/Plugins/VirtualProduction/Switchboard/Source/Switchboard/switchboard/switchboard_utils.py
+++ b/Engine/Plugins/VirtualProduction/Switchboard/Source/Switchboard/switchboard/switchboard_utils.py
@@ -189,17 +189,7 @@ def explore_path(path:str):
         subprocess.Popen(f'open {path}')
     else:
         LOGGER.error(f"explore_path not supported in platform '{sys.platform}'")
-
-def zip_files_in_folder(zipped_directory: str, zip_result_path: str, allow_file: typing.Callable[[str], bool]):
-    ''' Zips all files within zipped_directory and saves it as zip_result_path. Recursively iterates sub-folders.'''
-    with ZipFile(zip_result_path, "w") as zip_file:
-        for (dirpath, dirnames, filenames) in os.walk(zipped_directory):
-            for file_name in filenames:
-                file_path = os.path.join(dirpath, file_name)
-                if allow_file(file_path):
-                    directory_name_in_zip = pathlib.Path(zipped_directory).relative_to(pathlib.Path(dirpath))
-                    name_in_zip = os.path.join(directory_name_in_zip, file_name) if directory_name_in_zip else file_name 
-                    zip_file.write(file_path, name_in_zip)
+        
 
 def copy2clipboard(text:str):
     if sys.platform.startswith('win'):
diff --git a/Engine/Plugins/VirtualProduction/Switchboard/Source/Switchboard/switchboard/util/collect_logs.py b/Engine/Plugins/VirtualProduction/Switchboard/Source/Switchboard/switchboard/util/collect_logs.py
new file mode 100644
--- /dev/null
+++ b/Engine/Plugins/VirtualProduction/Switchboard/Source/Switchboard/switchboard/util/collect_logs.py
@@ -0,0 +1,119 @@
+# Copyright Epic Games, Inc. All Rights Reserved.
+
+import os
+import pathlib
+import typing
+from datetime import datetime
+from zipfile import ZipFile
+
+from PySide2 import QtWidgets
+
+from switchboard.config import Config
+from switchboard.devices.device_base import DeviceStatus
+from switchboard.devices.unreal.plugin_unreal import DeviceUnreal
+from ..switchboard_logging import DEFAULT_LOG_PATH, LOGGER
+
+from ..switchboard_utils import explore_path
+
+
+def open_logs_folder():
+    explore_path(DeviceUnreal.get_log_download_dir())
+    
+
+def execute_zip_logs_workflow(config: Config, devices: typing.List[DeviceUnreal]):
+    """ Zip up all important log files asking the user about preferences in certain edge cases. """
+    def should_skip_because_running_devices():
+        running_devices = [device for device in devices if
+                           device.status in [DeviceStatus.CONNECTING, DeviceStatus.OPEN, DeviceStatus.CLOSING]]
+        if len(running_devices) > 0:
+            pretty_devices_str = ""
+            for name in list(map(lambda device: device.name, running_devices)):
+                pretty_devices_str += f" - {name}\n"
+            continue_answer = QtWidgets.QMessageBox.question(
+                None,
+                "Devices still running",
+                f"The following devices are still running:\n{pretty_devices_str}"
+                "\nThese devices should be closed to have their latest logs included.\nDo you want to continue anyway?",
+                QtWidgets.QMessageBox.Yes | QtWidgets.QMessageBox.No
+            )
+
+            return True if continue_answer == QtWidgets.QMessageBox.No else False
+        return False
+    
+    def ask_log_zip_name():
+        modtime = datetime.now()
+        default_name = f"Logs-{modtime.strftime('%Y.%m.%d-%H.%M.%S')}"
+        result = QtWidgets.QInputDialog.getText(
+            None,
+            "Name .zip file",
+            "New .zip file name:",
+            QtWidgets.QLineEdit.Normal,
+            default_name
+        )
+
+        if not result[1]:
+            return None
+
+        file_path = os.path.join(DeviceUnreal.get_log_download_dir(), result[0])
+        if not file_path.endswith(".zip"):
+            file_path += ".zip"
+        return file_path
+    
+    if should_skip_because_running_devices():
+        return
+
+    zip_destination = ask_log_zip_name()
+    if zip_destination is not None:
+        zip_logs(zip_destination=zip_destination, config=config, devices=devices)
+
+
+def zip_logs(zip_destination: str, config: Config, devices: typing.List[DeviceUnreal]):
+    """ Collects all relevant logs and saves them.
+    """
+    def enqueue(abs_file_path: str, zip_path_name: str, zip_file: ZipFile):
+        # Files will be made relative to the engine directory
+        zip_file.write(abs_file_path, zip_path_name)
+    
+    with ZipFile(zip_destination, "w") as zip_file:
+        iterate_log_files(
+            config,
+            devices,
+            lambda abs_file_path, zip_path_name, zip_file=zip_file: enqueue(abs_file_path, zip_path_name, zip_file)
+        )
+        
+        
+def iterate_log_files(config: Config, devices: typing.List[DeviceUnreal], consumer: typing.Callable[[str, str], None]):
+    """ 
+        Passes the abolute files paths of all important log files:
+        - all latest .log files in Switchboard log folder
+        - all latest .utrace files in Switchboard log folder
+        - all multiuser session files in Programs/UnrealMultiUserServerIntermediateMultiUser
+        - multiuser log file ProgramsUnrealMultiUserServerSavedLogsUnrealMultiUserServer.log
+        
+        The arguments of the consumer:
+            1 The absolute file path of the file
+            2 A short unique relative file path (includes file name), useful for zipping
+    """
+    def iterate_file_list(files: typing.Iterable[str]):
+        for file_path in files:
+            if os.path.isfile(file_path):
+                consumer(file_path, os.path.join("Devices", os.path.basename(file_path))) 
+                
+    def iterate_mu_session_files():
+        for current_dir, _, files in os.walk(config.multiuser_server_session_directory_path()):
+            for file in files:
+                abs_path = os.path.join(current_dir, file)
+                root_mu_path = pathlib.Path(config.multiuser_server_session_directory_path())
+                relative_path = pathlib.Path(abs_path).relative_to(root_mu_path)
+                consumer(abs_path, os.path.join("MultiUserServer", relative_path))
+
+    latest_log_file_names = map(lambda device: device.last_log_path.get_value(), devices)
+    iterate_file_list(latest_log_file_names)
+
+    latest_trace_paths = map(lambda device: device.last_trace_path.get_value(), devices)
+    iterate_file_list(latest_trace_paths)
+
+    iterate_mu_session_files()
+    consumer(config.multiuser_server_log_path(), os.path.join("MultiUserServer", "UnrealMultiUserServer.log"))
+    
+    consumer(LOGGER.save_log_file(), "Switchboard.log")
\ No newline at end of file
