Synchronize access to multiuser database files

diff --git a/Engine/Plugins/VirtualProduction/Switchboard/Source/Switchboard/switchboard/util/collect_logs.py b/Engine/Plugins/VirtualProduction/Switchboard/Source/Switchboard/switchboard/util/collect_logs.py
--- a/Engine/Plugins/VirtualProduction/Switchboard/Source/Switchboard/switchboard/util/collect_logs.py
+++ b/Engine/Plugins/VirtualProduction/Switchboard/Source/Switchboard/switchboard/util/collect_logs.py
@@ -2,6 +2,7 @@
 
 import os
 import pathlib
+import sqlite3
 import typing
 from datetime import datetime
 from zipfile import ZipFile
@@ -70,17 +71,41 @@ def execute_zip_logs_workflow(config: Config, devices: typing.List[DeviceUnreal]
 def zip_logs(zip_destination: str, config: Config, devices: typing.List[DeviceUnreal]):
     """ Collects all relevant logs and saves them.
     """
-    def enqueue(abs_file_path: str, zip_path_name: str, zip_file: ZipFile):
-        # Files will be made relative to the engine directory
-        zip_file.write(abs_file_path, zip_path_name)
+    def process_log_file(abs_file_path: str, zip_path_name: str, zip_file: ZipFile):
+        # Access to .db files must be synchronized - these files are primarily used by multiuser
+        if abs_file_path.endswith(".db"):
+            copy_db_file(abs_file_path, zip_path_name, zip_file)
+        # Skip temporary files created by backup in copy_db_file
+        elif abs_file_path.endswith(".db-shm") or abs_file_path.endswith(".db-wal"):
+            return
+        else:
+            zip_file.write(abs_file_path, zip_path_name)
+        
+    def copy_db_file(abs_file_path: str, zip_path_name: str, zip_file: ZipFile):
+        copy_path = os.path.join(os.path.dirname(zip_destination), os.path.basename(zip_path_name))
+        try:
+            with sqlite3.connect(abs_file_path) as original, sqlite3.connect(copy_path) as copy:
+                open(copy_path, "a+").close()
+                # Cannot be removed before zip_file isn't closed
+                files_to_remove.append(copy_path)
+                original.backup(copy)
+                zip_file.write(copy_path, zip_path_name)
+        except sqlite3.Error:
+            LOGGER.error(f"Failed to access database file {abs_file_path}")
+        except (OSError, IOError): 
+            LOGGER.error(f"Failed to create temporary file {copy_path} while backing up original {abs_file_path}")
     
+    files_to_remove = []
     with ZipFile(zip_destination, "w") as zip_file:
         iterate_log_files(
             config,
             devices,
-            lambda abs_file_path, zip_path_name, zip_file=zip_file: enqueue(abs_file_path, zip_path_name, zip_file)
+            lambda abs_file_path, zip_path_name, zip_file=zip_file: process_log_file(abs_file_path, zip_path_name, zip_file)
         )
         
+    for file_to_remove in files_to_remove:
+        os.remove(file_to_remove)
+        
         
 def iterate_log_files(config: Config, devices: typing.List[DeviceUnreal], consumer: typing.Callable[[str, str], None]):
     """ 
