Switchboard: Fixes to enable running on Linux. Uses a higher default port for rsync which does not require root privileges. Also includes a fix for a race condition starting new short-lived processes in SwitchboardListener.

diff --git a/Engine/Plugins/VirtualProduction/Switchboard/Source/Switchboard/switchboard/p4_utils.py b/Engine/Plugins/VirtualProduction/Switchboard/Source/Switchboard/switchboard/p4_utils.py
--- a/Engine/Plugins/VirtualProduction/Switchboard/Source/Switchboard/switchboard/p4_utils.py
+++ b/Engine/Plugins/VirtualProduction/Switchboard/Source/Switchboard/switchboard/p4_utils.py
@@ -30,7 +30,7 @@ def p4_latest_changelist(p4_path, working_dir, num_changelists=10):
     p4_command = f'p4 -ztag -F "%change%" changes -m {num_changelists} {p4_path}/...'
     LOGGER.info(f"Executing: {p4_command}")
 
-    p4_result = subprocess.check_output(p4_command, cwd=working_dir, startupinfo=sb_utils.get_hidden_sp_startupinfo()).decode()
+    p4_result = subprocess.check_output(p4_command, cwd=working_dir, shell=True, startupinfo=sb_utils.get_hidden_sp_startupinfo()).decode()
 
     if p4_result:
         return p4_result.split()
diff --git a/Engine/Plugins/VirtualProduction/Switchboard/Source/Switchboard/switchboard/switchboard_application.py b/Engine/Plugins/VirtualProduction/Switchboard/Source/Switchboard/switchboard/switchboard_application.py
--- a/Engine/Plugins/VirtualProduction/Switchboard/Source/Switchboard/switchboard/switchboard_application.py
+++ b/Engine/Plugins/VirtualProduction/Switchboard/Source/Switchboard/switchboard/switchboard_application.py
@@ -3,6 +3,7 @@
 import os
 import pathlib
 import subprocess
+import sys
 import tempfile
 import threading
 from typing import Dict, List, NamedTuple, Optional
@@ -173,7 +174,7 @@ def get_multi_user_server_instance():
     return MultiUserServerInstance
 
 class RsyncServer:
-    DEFAULT_PORT = 873
+    DEFAULT_PORT = 8730
     INCOMING_LOGS_MODULE = 'device_logs'
     MAX_MONITOR_RECOVERIES = 5
     SOCKET_ERROR_EXIT_CODE = 10
@@ -284,7 +285,7 @@ hosts allow = {allowed_addrs}
     def launch(
         self, *, address: str = '0.0.0.0', port: int = DEFAULT_PORT
     ) -> bool:
-        if self.process is not None:
+        if self.is_running():
             LOGGER.warning('RsyncServer.launch: server already running')
             return False
 
@@ -292,9 +293,12 @@ hosts allow = {allowed_addrs}
         self.port = port
         self.update_config()
 
-        rsync_path = os.path.normpath(os.path.join(
-            CONFIG.ENGINE_DIR.get_value(), 'Extras', 'ThirdPartyNotUE',
-            'SwitchboardThirdParty', 'cwrsync', 'bin', 'rsync.exe'))
+        if sys.platform.startswith('win'):
+            rsync_path = os.path.normpath(os.path.join(
+                CONFIG.ENGINE_DIR.get_value(), 'Extras', 'ThirdPartyNotUE',
+                'SwitchboardThirdParty', 'cwrsync', 'bin', 'rsync.exe'))
+        else:
+            rsync_path = 'rsync'
 
         args = [rsync_path, '--daemon', '--no-detach', f'--address={address}',
                 f'--port={port}', f'--config={self.config_file.name}',
diff --git a/Engine/Plugins/VirtualProduction/Switchboard/Source/Switchboard/switchboard/switchboard_utils.py b/Engine/Plugins/VirtualProduction/Switchboard/Source/Switchboard/switchboard/switchboard_utils.py
--- a/Engine/Plugins/VirtualProduction/Switchboard/Source/Switchboard/switchboard/switchboard_utils.py
+++ b/Engine/Plugins/VirtualProduction/Switchboard/Source/Switchboard/switchboard/switchboard_utils.py
@@ -28,9 +28,10 @@ class PriorityModifier(Enum):
 
 def get_hidden_sp_startupinfo():
     ''' Returns subprocess.startupinfo and avoids extra cmd line window in Windows. '''
-    startupinfo = subprocess.STARTUPINFO()
+    startupinfo = None
 
-    if sys.platform.startswith("win"):
+    if hasattr(subprocess, 'STARTUPINFO') and hasattr(subprocess, 'STARTF_USESHOWWINDOW'):
+        startupinfo = subprocess.STARTUPINFO()
         startupinfo.dwFlags |= subprocess.STARTF_USESHOWWINDOW
 
     return startupinfo
diff --git a/Engine/Source/Programs/SwitchboardListener/Private/SwitchboardListener.cpp b/Engine/Source/Programs/SwitchboardListener/Private/SwitchboardListener.cpp
--- a/Engine/Source/Programs/SwitchboardListener/Private/SwitchboardListener.cpp
+++ b/Engine/Source/Programs/SwitchboardListener/Private/SwitchboardListener.cpp
@@ -650,7 +650,7 @@ bool FSwitchboardListener::Task_StartProcess(const FSwitchboardStartTask& InRunT
 		NewProcess->ReadPipe
 	);
 
-	if (!NewProcess->Handle.IsValid() || !FPlatformProcess::IsProcRunning(NewProcess->Handle))
+	if (!NewProcess->Handle.IsValid())
 	{
 		// Close process in case it just didn't run
 		FPlatformProcess::CloseProc(NewProcess->Handle);
