Synching content only changelists will no longer indicate that a rebuild is required

diff --git a/Engine/Plugins/VirtualProduction/Switchboard/Source/Switchboard/switchboard/devices/unreal/plugin_unreal.py b/Engine/Plugins/VirtualProduction/Switchboard/Source/Switchboard/switchboard/devices/unreal/plugin_unreal.py
--- a/Engine/Plugins/VirtualProduction/Switchboard/Source/Switchboard/switchboard/devices/unreal/plugin_unreal.py
+++ b/Engine/Plugins/VirtualProduction/Switchboard/Source/Switchboard/switchboard/devices/unreal/plugin_unreal.py
@@ -37,6 +37,9 @@ from switchboard.tools.insights_launcher import InsightsLauncher
 from .listener_watcher import ListenerWatcher
 from .redeploy_dialog import RedeployListenerDialog
 from . import version_helpers
+from ...util import p4_changelist_inspection
+from ...util.p4_changelist_inspection import P4Error
+
 
 class ProgramStartQueueItem:
     '''
@@ -2344,20 +2347,20 @@ class DeviceWidgetUnreal(DeviceWidget):
 
         self.project_changelist_label.show()
 
-    def update_engine_changelist(self, required_cl: str, current_device_cl: str, built_device_cl: str):
+    def update_engine_changelist(self, required_cl: str, synched_cl: str, built__cl: str):
         if not CONFIG.BUILD_ENGINE.get_value():
             return
         
-        self.engine_changelist_label.setText(f'E: {current_device_cl}')
+        self.engine_changelist_label.setText(f'E: {synched_cl}')
         if not self.exclude_from_build:
             self.engine_changelist_label.show()
 
         self.sync_button.show()
         self.build_button.show()
         
-        is_synched = required_cl is None or required_cl == current_device_cl
+        is_synched = required_cl is None or required_cl == synched_cl
         self._set_engine_changelist_is_synched(is_synched)
-        self.update_build_info(current_cl=current_device_cl, built_cl=built_device_cl)
+        self.update_build_info(synched_cl=synched_cl, built_cl=built__cl)
 
     def _set_project_changelist_is_synched(self, is_synched: bool):
         self._is_project_synched = is_synched
@@ -2377,19 +2380,28 @@ class DeviceWidgetUnreal(DeviceWidget):
         needs_resync = not self._is_project_synched or not self._is_engine_synched
         sb_widgets.set_qt_property(self.sync_button, 'not_synched', needs_resync)
             
-    def update_build_info(self, current_cl: str, built_cl: str):
-        def set_desired_tooltip(current_cl: str, built_cl: str):
-            desired_tooltip = f"Build changelist.\n\nBuild required.\nCurrent: {current_cl}\nBuilt: {built_cl}" \
-                if self._needs_rebuild else f"Build changelist (not required - built CL {built_cl} matches synched CL)"
-            self._desired_build_button_tooltip = desired_tooltip
-            self._update_build_button_tooltip()
-        
-        if built_cl is not None and current_cl is not None and CONFIG.BUILD_ENGINE.get_value():
-            self._needs_rebuild = int(built_cl) != int(current_cl)
+    def update_build_info(self, synched_cl: str, built_cl: str):
+        if built_cl is not None and synched_cl is not None and CONFIG.BUILD_ENGINE.get_value():
+            try:
+                earlier_cl = min(int(built_cl), int(synched_cl))
+                later_cl = max(int(built_cl), int(synched_cl))
+                self._needs_rebuild = p4_changelist_inspection.has_source_code_changes(
+                    earlier_cl,
+                    later_cl,
+                    CONFIG.P4_ENGINE_PATH.get_value()
+                )
+            except P4Error as error:
+                LOGGER.error(f"Couldn't check {built_cl} - {built_cl} for source code changes."
+                             f" Reason: {error.message}")
+                # Assume that non-equal CL numbers have source changes
+                self._needs_rebuild = int(built_cl) != int(synched_cl)
         else:
             self._needs_rebuild = False
 
-        set_desired_tooltip(current_cl, built_cl)
+        desired_tooltip = f"Build changelist.\n\nBuild required.\nSynched: {synched_cl}\nBuilt: {built_cl}" \
+            if self._needs_rebuild else f"Build changelist (not required - built CL {built_cl} matches synched CL)"
+        self._desired_build_button_tooltip = desired_tooltip
+        self._update_build_button_tooltip()
         self._refresh_build_info_ui()
         
     def _refresh_build_info_ui(self):
diff --git a/Engine/Plugins/VirtualProduction/Switchboard/Source/Switchboard/switchboard/p4_utils.py b/Engine/Plugins/VirtualProduction/Switchboard/Source/Switchboard/switchboard/p4_utils.py
--- a/Engine/Plugins/VirtualProduction/Switchboard/Source/Switchboard/switchboard/p4_utils.py
+++ b/Engine/Plugins/VirtualProduction/Switchboard/Source/Switchboard/switchboard/p4_utils.py
@@ -68,7 +68,7 @@ def run(cmd, args=[], input=None):
     return r
 
 
-def valueForMarshalledKey(obj, key):
+def valueForMarshalledKey(obj: dict, key: str):
     ''' The P4 marshal is using bytes as keys instead of strings,
     so this makes the conversion and returns the desired value for the given key.
 
@@ -79,6 +79,14 @@ def valueForMarshalledKey(obj, key):
     return obj[key.encode('utf-8')].decode()
 
 
+def hasValueForMarshalledKey(obj: dict, key: str):
+    '''
+    Checks whether a key exists.
+    See valueForMarshalledKey.
+    '''
+    return key.encode('utf-8') in obj
+
+
 def workspaceInPath(ws, localpath):
     ''' Validates if the give localpath can correspond to the given workspace.
 
diff --git a/Engine/Plugins/VirtualProduction/Switchboard/Source/Switchboard/switchboard/switchboard_dialog.py b/Engine/Plugins/VirtualProduction/Switchboard/Source/Switchboard/switchboard/switchboard_dialog.py
--- a/Engine/Plugins/VirtualProduction/Switchboard/Source/Switchboard/switchboard/switchboard_dialog.py
+++ b/Engine/Plugins/VirtualProduction/Switchboard/Source/Switchboard/switchboard/switchboard_dialog.py
@@ -1310,8 +1310,8 @@ class SwitchboardDialog(QtCore.QObject):
                 device_widget = self.device_list_widget.device_widget_by_hash(device.device_hash)
                 device_widget.update_engine_changelist(
                     required_cl=self.engine_changelist if self.engine_changelist != EMPTY_SYNC_ENTRY else None, 
-                    current_device_cl=device.engine_changelist,
-                    built_device_cl=device.built_engine_changelist
+                    synched_cl=device.engine_changelist,
+                    built__cl=device.built_engine_changelist
                 )
 
     @QtCore.Slot(object)
@@ -1424,8 +1424,8 @@ class SwitchboardDialog(QtCore.QObject):
         device_widget = self.device_list_widget.device_widget_by_hash(device.device_hash)
         device_widget.update_engine_changelist(
             required_cl=self.engine_changelist if self.engine_changelist != EMPTY_SYNC_ENTRY else None,
-            current_device_cl=device.engine_changelist,
-            built_device_cl=device.built_engine_changelist
+            synched_cl=device.engine_changelist,
+            built__cl=device.built_engine_changelist
         )
 
         cl = device.engine_changelist
@@ -1437,7 +1437,10 @@ class SwitchboardDialog(QtCore.QObject):
     @QtCore.Slot(object)
     def device_built_engine_changelist_changed(self, device):
         device_widget = self.device_list_widget.device_widget_by_hash(device.device_hash)
-        device_widget.update_build_info(current_cl=device.engine_changelist, built_cl=device.built_engine_changelist)
+        device_widget.update_build_info(
+            synched_cl=device.engine_changelist,
+            built_cl=device.built_engine_changelist
+        )
 
     @QtCore.Slot(object)
     def device_status_changed(self, device, previous_status):
diff --git a/Engine/Plugins/VirtualProduction/Switchboard/Source/Switchboard/switchboard/util/p4_changelist_inspection.py b/Engine/Plugins/VirtualProduction/Switchboard/Source/Switchboard/switchboard/util/p4_changelist_inspection.py
new file mode 100644
--- /dev/null
+++ b/Engine/Plugins/VirtualProduction/Switchboard/Source/Switchboard/switchboard/util/p4_changelist_inspection.py
@@ -0,0 +1,104 @@
+# Copyright Epic Games, Inc. All Rights Reserved.
+import os
+import subprocess
+import typing
+from enum import IntFlag
+
+from ..config import CONFIG
+from ..p4_utils import run, valueForMarshalledKey, hasValueForMarshalledKey
+from ..switchboard_utils import get_hidden_sp_startupinfo
+
+
+def has_source_code_changes(earlier_cl: int, later_cl: int, engine_path: str) -> bool:
+    """
+    Checks whether the CLs between current_cl (exclusive) and future_cl (inclusive) are content only.
+    
+    Raises P4Error if the data cannot be obtained.
+    """
+    if earlier_cl == later_cl:
+        return False
+    
+    return any_of_changelists(
+        earlier_cl,
+        later_cl,
+        engine_path,
+        lambda cl_description:
+            EChangelistFileFlags.Code in determine_changelist_type(valueForMarshalledKey(cl_description, "change")),
+        [earlier_cl]
+    )
+
+def any_of_changelists(
+    earlier_cl: int,
+    later_cl: int,
+    engine_path: str,
+    predicate_callback: typing.Callable[[typing.Dict], bool],
+    skipped_cls: typing.List[int] = None
+):
+    """
+    Checks whether a predicate holds for all CLs from current_cl to future_cl (both inclusive),
+    The callback can accept the result of an entry returned by "p4 -z tag -G changes".
+    """
+    # Maybe we need to use -m 10 here to only get 10 CLs at a time... if there are a lot of changes p4 may hang up
+    changes: dict = run(f"changes {engine_path}/...@{earlier_cl},{later_cl}")
+    _raise_exception_on_error(changes)
+    
+    for change in changes:
+        cl_number = int(valueForMarshalledKey(change, "change"))
+        should_skip = skipped_cls is not None and cl_number in skipped_cls
+        if not should_skip and predicate_callback(change):
+            return True
+        
+    return False
+
+
+class EChangelistFileFlags(IntFlag):
+    # CL contains source changes
+    Code = 1,
+    # CL contains content changes
+    Content = 2
+    
+    
+def determine_changelist_type(changelist: int) -> EChangelistFileFlags:
+    compiled_ue_source_extensions = ["c", "cs", "cpp", "h", "hpp"]
+    
+    changes: dict = run(f"describe {changelist}")
+    _raise_exception_on_error(changes)
+    
+    has_source_changes = False
+    has_content_changes = False
+    for change in changes:
+        file_index: int = 0
+        # Each changed file path name will have an entry depotFilex
+        while hasValueForMarshalledKey(change, f"depotFile{file_index}"):
+            file_path_name: str = valueForMarshalledKey(change, f"depotFile{file_index}")
+            
+            has_source_changes = any([file_path_name.endswith(extension) for extension in compiled_ue_source_extensions])
+            # If it's not a source file, treat it as a content file
+            has_content_changes = not has_source_changes
+            
+            # All flags fulfilled - no need to inspect more files
+            if has_source_changes and has_content_changes:
+                break
+            
+            file_index += 1
+    
+    result = 0
+    if has_source_changes:
+        result |= EChangelistFileFlags.Code
+    if has_content_changes:
+        result |= EChangelistFileFlags.Content
+    return result
+
+
+class P4Error(RuntimeError):
+    def __init__(self, message: str):
+        self.message = message
+
+
+def _raise_exception_on_error(p4_command_result: list):
+    """
+    Accepts the result of p4_util.run and raises an error if any error is set
+    """
+    if len(p4_command_result) > 0 and valueForMarshalledKey(p4_command_result[0], "code") == "error":
+        error_message = valueForMarshalledKey(p4_command_result[0], "data")
+        raise P4Error(error_message)
\ No newline at end of file
