SwitchboardListener: Follow-up for incomplete character encoding fix, bump Listener version.

Also refactors the Switchboard ListenerClient hitch detection code slightly to make a magic number more explicit.

diff --git a/Engine/Plugins/VirtualProduction/Switchboard/Source/Switchboard/switchboard/listener_client.py b/Engine/Plugins/VirtualProduction/Switchboard/Source/Switchboard/switchboard/listener_client.py
--- a/Engine/Plugins/VirtualProduction/Switchboard/Source/Switchboard/switchboard/listener_client.py
+++ b/Engine/Plugins/VirtualProduction/Switchboard/Source/Switchboard/switchboard/listener_client.py
@@ -33,6 +33,10 @@ class ListenerClient(object):
     '''
 
     KEEPALIVE_INTERVAL_SEC = 1.0
+    SELECT_TIMEOUT_SEC = 0.1
+    EXTRA_HITCH_TOLERANCE_SEC = 0.4
+    HITCH_THRESHOLD_SEC = (KEEPALIVE_INTERVAL_SEC + SELECT_TIMEOUT_SEC
+                           + EXTRA_HITCH_TOLERANCE_SEC)
 
     def __init__(self, ip_address, port, buffer_size=1024):
         self.ip_address = ip_address
@@ -150,7 +154,7 @@ class ListenerClient(object):
         now = datetime.datetime.now()
         delta_sec = (now - self.last_activity).total_seconds()
         self.last_activity = now
-        if delta_sec > (self.KEEPALIVE_INTERVAL_SEC * 1.5):
+        if delta_sec > self.HITCH_THRESHOLD_SEC:
             LOGGER.warning(
                 f'Hitch detected; {delta_sec:.1f} seconds since last keepalive'
                 f' to {self.ip_address}')
@@ -163,9 +167,9 @@ class ListenerClient(object):
                 rlist = [self.socket]
                 wlist = []
                 xlist = []
-                read_timeout = 0.1
 
-                read_sockets, _, _ = select.select(rlist, wlist, xlist, read_timeout)
+                read_sockets, _, _ = select.select(rlist, wlist, xlist,
+                                                   self.SELECT_TIMEOUT_SEC)
 
                 while len(self.message_queue):
                     message_bytes = self.message_queue.pop()
diff --git a/Engine/Source/Programs/SwitchboardListener/Private/SwitchboardListener.cpp b/Engine/Source/Programs/SwitchboardListener/Private/SwitchboardListener.cpp
--- a/Engine/Source/Programs/SwitchboardListener/Private/SwitchboardListener.cpp
+++ b/Engine/Source/Programs/SwitchboardListener/Private/SwitchboardListener.cpp
@@ -1615,7 +1615,7 @@ static void FillOutFlipMode(FSyncStatus& SyncStatus, FRunningProcess* FlipModeMo
 
 	// Get stdout.
 	FUTF8ToTCHAR StdoutConv(reinterpret_cast<const ANSICHAR*>(FlipModeMonitor->Output.GetData()), FlipModeMonitor->Output.Num());
-	const FString StdOut(StdoutConv.Get());
+	const FString StdOut(StdoutConv.Length(), StdoutConv.Get());
 
 	// Clear out the StdOut array.
 	FlipModeMonitor->Output.Empty();
@@ -2149,7 +2149,7 @@ void FSwitchboardListener::HandleRunningProcesses(TArray<TSharedPtr<FRunningProc
 				if (ReturnCode != 0 && Process->Output.Num() > 0)
 				{
 					FUTF8ToTCHAR StdoutConv(reinterpret_cast<const ANSICHAR*>(Process->Output.GetData()), Process->Output.Num());
-					UE_LOG(LogSwitchboard, Display, TEXT("Output:\n%s"), StdoutConv.Get());
+					UE_LOG(LogSwitchboard, Display, TEXT("Output:\n%.*s"), StdoutConv.Length(), StdoutConv.Get());
 				}
 
 				// Notify remote client, which implies that this is a program managed by it.
diff --git a/Engine/Source/Programs/SwitchboardListener/Private/SwitchboardListenerVersion.h b/Engine/Source/Programs/SwitchboardListener/Private/SwitchboardListenerVersion.h
--- a/Engine/Source/Programs/SwitchboardListener/Private/SwitchboardListenerVersion.h
+++ b/Engine/Source/Programs/SwitchboardListener/Private/SwitchboardListenerVersion.h
@@ -4,4 +4,4 @@
 
 #define SBLISTENER_VERSION_MAJOR 2
 #define SBLISTENER_VERSION_MINOR 0
-#define SBLISTENER_VERSION_PATCH 0
+#define SBLISTENER_VERSION_PATCH 1
