Switchboard: Improved onboarding.

The Python virtual environment path can be configured in the editor preferences, and the integrity of the venv is automatically checked and reflected in the UI. The Switchboard project settings now only contain the OSC listener editor utility blueprint, and a link to the editor preferences.

Remaining work:
- Toolbar icon still shows all the time, and launches Switchboard rather than the onboarding process.
- The create shortcuts step doesn't actually create shortcuts yet.

diff --git a/Engine/Plugins/VirtualProduction/Switchboard/Source/Switchboard/sb_setup.py b/Engine/Plugins/VirtualProduction/Switchboard/Source/Switchboard/sb_setup.py
--- a/Engine/Plugins/VirtualProduction/Switchboard/Source/Switchboard/sb_setup.py
+++ b/Engine/Plugins/VirtualProduction/Switchboard/Source/Switchboard/sb_setup.py
@@ -158,7 +158,11 @@ class SbSetup:
         if sys.platform.startswith('win'):
             self.ensure_rsync_copied()
 
-        venv_dir = options.venv_dir
+        venv_dir: pathlib.Path = options.venv_dir
+
+        # abspath is workaround for https://bugs.python.org/issue38671
+        venv_dir = pathlib.Path(os.path.abspath(venv_dir.resolve()))
+
         logging.info(f'VENV_DIR: {venv_dir}')
 
         dest_not_empty = False
@@ -265,10 +269,7 @@ class SbSetup:
                 # Output results to stdout in JSON format
                 print(json.dumps(verify_result, indent=4))
 
-            if verify_result['success']:
-                return 0
-            else:
-                return 1
+            return 0
         else:
             # Shouldn't happen: argparse early outs on unregistered actions.
             assert False
diff --git a/Engine/Plugins/VirtualProduction/Switchboard/Source/Switchboard/switchboard.bat b/Engine/Plugins/VirtualProduction/Switchboard/Source/Switchboard/switchboard.bat
--- a/Engine/Plugins/VirtualProduction/Switchboard/Source/Switchboard/switchboard.bat
+++ b/Engine/Plugins/VirtualProduction/Switchboard/Source/Switchboard/switchboard.bat
@@ -1,14 +1,33 @@
 @echo off
 setlocal ENABLEDELAYEDEXPANSION
 
+REM This script optionally takes a single argument, representing the path to the desired Python
+REM virtual environment directory. If omitted, it defaults to the value of %_defaultVenvDir%.
+REM
+REM Additionally, in the case where no argument was provided, pause on errors, in case the user
+REM double-clicked the batch file and we don't want the console window to vanish inexplicably.
+REM Otherwise, we exit immediately.
+
+
 set _switchboardDir=%~dp0
 REM pushd and %CD% are used to normalize the relative path to a shorter absolute path.
 pushd "%~dp0..\..\..\..\.."
 set _engineDir=%CD%
 set _enginePythonPlatformDir=%_engineDir%\Binaries\ThirdParty\Python3\Win64
-set _pyVenvDir=%_engineDir%\Extras\ThirdPartyNotUE\SwitchboardThirdParty\Python
+set _defaultVenvDir=%_engineDir%\Extras\ThirdPartyNotUE\SwitchboardThirdParty\Python
 popd
 
+if not "%~1"=="" (
+    set _bUsingDefaultVenv=
+    set _venvDir=%~1
+    echo Using provided path for Python virtual environment (!_venvDir!^)
+) else (
+    REM Called with no arguments; assume user double-clicked on us and pause for errors.
+    set _bUsingDefaultVenv=1
+    set _venvDir=%_defaultVenvDir%
+    echo Using DEFAULT path for Python virtual environment (%_defaultVenvDir%^)
+)
+
 call:main
 
 endlocal
@@ -17,12 +36,12 @@ goto:eof
 ::------------------------------------------------------------------------------
 :main
 
-if not exist "%_pyVenvDir%\Scripts\pythonw.exe" (
-    echo Performing Switchboard first-time setup (using default path for Python virtual environment^)
-    call "%_enginePythonPlatformDir%\python.exe" "%~dp0\sb_setup.py" install --venv-dir="%_pyVenvDir%"
+if not exist "%_venvDir%\Scripts\pythonw.exe" (
+    echo Performing Switchboard first-time setup
+    call "%_enginePythonPlatformDir%\python.exe" "%~dp0\sb_setup.py" install --venv-dir="%_venvDir%"
     if !ERRORLEVEL! NEQ 0 (
         echo Installation failed with non-zero exit code^^^!
-        pause
+        if defined _bUsingDefaultVenv ( pause )
         exit /B !ERRORLEVEL!
     )
 )
@@ -35,8 +54,9 @@ goto:eof
 :start_sb
 
 set PYTHONPATH=%_switchboardDir%;%PYTHONPATH%
-start "Switchboard" "%_pyVenvDir%\Scripts\pythonw.exe" -m switchboard
+start "Switchboard" "%_venvDir%\Scripts\pythonw.exe" -m switchboard
 if %ERRORLEVEL% NEQ 0 (
     echo Failed to launch Switchboard^!
-    pause
+    if defined _bUsingDefaultVenv ( pause )
+    exit /B %ERRORLEVEL%
 )
diff --git a/Engine/Plugins/VirtualProduction/Switchboard/Source/Switchboard/switchboard.sh b/Engine/Plugins/VirtualProduction/Switchboard/Source/Switchboard/switchboard.sh
--- a/Engine/Plugins/VirtualProduction/Switchboard/Source/Switchboard/switchboard.sh
+++ b/Engine/Plugins/VirtualProduction/Switchboard/Source/Switchboard/switchboard.sh
@@ -1,9 +1,18 @@
 #!/bin/sh
 
+# This script optionally takes a single argument, representing the path to the desired Python
+# virtual environment directory. If omitted, it defaults to the value of $_defaultVenvDir.
+
 _switchboardDir=$(dirname "$0")
 _engineDir=$(cd "$_switchboardDir/../../../../.."; pwd)
 _enginePythonDir="$_engineDir/Binaries/ThirdParty/Python3/Linux"
-_venvDir="$_engineDir/Extras/ThirdPartyNotUE/SwitchboardThirdParty/Python"
+_defaultVenvDir="$_engineDir/Extras/ThirdPartyNotUE/SwitchboardThirdParty/Python"
+
+if [ $# -lt 1 ]; then
+    _venvDir=$_defaultVenvDir
+else
+    _venvDir="$1"
+fi
 
 if [ ! -x "$_venvDir/bin/python3" ]; then
     "$_enginePythonDir/bin/python3" "$_switchboardDir/sb_setup.py" install --venv-dir="$_venvDir"
diff --git a/Engine/Plugins/VirtualProduction/Switchboard/Source/SwitchboardEditor/Private/SwitchboardEditorModule.cpp b/Engine/Plugins/VirtualProduction/Switchboard/Source/SwitchboardEditor/Private/SwitchboardEditorModule.cpp
--- a/Engine/Plugins/VirtualProduction/Switchboard/Source/SwitchboardEditor/Private/SwitchboardEditorModule.cpp
+++ b/Engine/Plugins/VirtualProduction/Switchboard/Source/SwitchboardEditor/Private/SwitchboardEditorModule.cpp
@@ -1,15 +1,16 @@
 // Copyright Epic Games, Inc. All Rights Reserved.
 
 #include "SwitchboardEditorModule.h"
-
+#include "SwitchboardEditorSettings.h"
+#include "SwitchboardMenuEntry.h"
+#include "SwitchboardProjectSettings.h"
+#include "SwitchboardSettingsCustomization.h"
 #include "Editor.h"
 #include "EditorUtilitySubsystem.h"
 #include "ISettingsModule.h"
 #include "ISettingsSection.h"
 #include "Misc/CoreDelegates.h"
 #include "Misc/ScopeExit.h"
-#include "SwitchboardEditorSettings.h"
-#include "SwitchboardMenuEntry.h"
 
 #if PLATFORM_WINDOWS
 #include "Windows/AllowWindowsPlatformTypes.h"
@@ -25,27 +26,62 @@ bool SetListenerAutolaunchEntry(const FString& NewCommandLine);
 bool RemoveListenerAutolaunchEntry();
 #endif
 
+
 #define LOCTEXT_NAMESPACE "SwitchboardEditorModule"
 
 DEFINE_LOG_CATEGORY(LogSwitchboardPlugin);
 
 
+// static
+const FString& FSwitchboardEditorModule::GetSbScriptsPath()
+{
+	using namespace UE::Switchboard::Private;
+	static FString SbScriptsPath = ConcatPaths(FPaths::EnginePluginsDir(),
+		"VirtualProduction", "Switchboard", "Source", "Switchboard");
+	return SbScriptsPath;
+}
+
+
+// static
+const FString& FSwitchboardEditorModule::GetSbThirdPartyPath()
+{
+	using namespace UE::Switchboard::Private;
+	static FString SbThirdPartyPath = ConcatPaths(FPaths::EngineDir(),
+		"Extras", "ThirdPartyNotUE", "SwitchboardThirdParty");
+	return SbThirdPartyPath;
+}
+
+
 void FSwitchboardEditorModule::StartupModule()
 {
 	FSwitchboardMenuEntry::Register();
 
 	if (ISettingsModule* SettingsModule = FModuleManager::GetModulePtr<ISettingsModule>("Settings"))
 	{
-		ISettingsSectionPtr SettingsSection = SettingsModule->RegisterSettings("Project", "Plugins", "Switchboard",
-			LOCTEXT("SettingsName", "Switchboard"),
-			LOCTEXT("Description", "Configure the Switchboard launcher."),
+		ISettingsSectionPtr EditorSettingsSection = SettingsModule->RegisterSettings("Editor", "Plugins", "Switchboard",
+			LOCTEXT("EditorSettingsName", "Switchboard"),
+			LOCTEXT("EditorSettingsDescription", "Configure the Switchboard launcher."),
 			GetMutableDefault<USwitchboardEditorSettings>()
 		);
 
-		SettingsSection->OnModified().BindRaw(this, &FSwitchboardEditorModule::OnSettingsModified);
+		ISettingsSectionPtr ProjectSettingsSection = SettingsModule->RegisterSettings("Project", "Plugins", "Switchboard",
+			LOCTEXT("ProjectSettingsName", "Switchboard"),
+			LOCTEXT("ProjectSettingsDescription", "Configure the Switchboard launcher."),
+			GetMutableDefault<USwitchboardProjectSettings>()
+		);
+
+		EditorSettingsSection->OnModified().BindRaw(this, &FSwitchboardEditorModule::OnEditorSettingsModified);
 	}
 
+	FPropertyEditorModule& PropertyModule = FModuleManager::GetModuleChecked<FPropertyEditorModule>("PropertyEditor");
+	PropertyModule.RegisterCustomClassLayout(USwitchboardEditorSettings::StaticClass()->GetFName(), FOnGetDetailCustomizationInstance::CreateStatic(&FSwitchboardEditorSettingsCustomization::MakeInstance));
+	PropertyModule.RegisterCustomClassLayout(USwitchboardProjectSettings::StaticClass()->GetFName(), FOnGetDetailCustomizationInstance::CreateStatic(&FSwitchboardProjectSettingsCustomization::MakeInstance));
+
 	DeferredStartDelegateHandle = FCoreDelegates::OnFEngineLoopInitComplete.AddRaw(this, &FSwitchboardEditorModule::OnEngineInitComplete);
+
+#if SB_LISTENER_AUTOLAUNCH
+	bCachedAutolaunchEnabled = GetListenerAutolaunchEnabled_Internal();
+#endif
 }
 
 
@@ -53,18 +89,98 @@ void FSwitchboardEditorModule::ShutdownModule()
 {
 	if (!IsRunningCommandlet() && UObjectInitialized() && !IsEngineExitRequested())
 	{
-		ISettingsModule* SettingsModule = FModuleManager::GetModulePtr<ISettingsModule>("Settings");
-		if (SettingsModule != nullptr)
+		FPropertyEditorModule& PropertyModule = FModuleManager::GetModuleChecked<FPropertyEditorModule>("PropertyEditor");
+		PropertyModule.UnregisterCustomClassLayout(USwitchboardEditorSettings::StaticClass()->GetFName());
+		PropertyModule.UnregisterCustomClassLayout(USwitchboardProjectSettings::StaticClass()->GetFName());
+
+		if (ISettingsModule* SettingsModule = FModuleManager::GetModulePtr<ISettingsModule>("Settings"))
 		{
+			SettingsModule->UnregisterSettings("Editor", "Plugins", "Switchboard");
 			SettingsModule->UnregisterSettings("Project", "Plugins", "Switchboard");
 		}
+
 		FSwitchboardMenuEntry::Unregister();
 	}
 }
 
 
+bool FSwitchboardEditorModule::LaunchSwitchboard()
+{
 #if PLATFORM_WINDOWS
-bool FSwitchboardEditorModule::IsListenerAutolaunchEnabled() const
+	const FString ScriptPath = GetSbScriptsPath() / TEXT("switchboard.bat");
+#elif PLATFORM_LINUX
+	const FString ScriptPath = GetSbScriptsPath() / TEXT("switchboard.sh");
+#endif
+
+	const FString ScriptArgs = FString::Printf(TEXT("\"%s\""),
+		*GetDefault<USwitchboardEditorSettings>()->VirtualEnvironmentPath.Path);
+
+	return RunProcess(ScriptPath, ScriptArgs);
+}
+
+
+bool FSwitchboardEditorModule::LaunchListener()
+{
+	const FString ListenerPath = GetDefault<USwitchboardEditorSettings>()->GetListenerPlatformPath();
+	const FString ListenerArgs = GetDefault<USwitchboardEditorSettings>()->ListenerCommandlineArguments;
+	return RunProcess(ListenerPath, ListenerArgs);
+}
+
+
+bool FSwitchboardEditorModule::RunProcess(const FString& InExe, const FString& InArgs)
+{
+	const bool bLaunchDetached = false;
+	const bool bLaunchHidden = false;
+	const bool bLaunchReallyHidden = false;
+	uint32* OutProcessId = nullptr;
+	const int32 PriorityModifier = 0;
+	const TCHAR* WorkingDirectory = nullptr;
+	void* PipeWriteChild = nullptr;
+
+	const FProcHandle Handle = FPlatformProcess::CreateProc(*InExe, *InArgs, bLaunchDetached, bLaunchHidden, bLaunchReallyHidden, OutProcessId, PriorityModifier, WorkingDirectory, PipeWriteChild);
+	return Handle.IsValid();
+}
+
+
+TSharedFuture<FSwitchboardVerifyResult> FSwitchboardEditorModule::GetVerifyResult()
+{
+	const FString CurrentVenv = GetDefault<USwitchboardEditorSettings>()->VirtualEnvironmentPath.Path;
+
+	if (!VerifyResult.IsValid() || VerifyPath != CurrentVenv)
+	{
+		UE_LOG(LogSwitchboardPlugin, Log, TEXT("Issuing verify for venv: %s"), *CurrentVenv);
+
+		TFuture<FSwitchboardVerifyResult> Future = FSwitchboardVerifyResult::RunVerify(CurrentVenv);
+		Future = Future.Next([Venv = CurrentVenv](const FSwitchboardVerifyResult& Result)
+		{
+			UE_LOG(LogSwitchboardPlugin, Log, TEXT("Verify complete for venv: %s"), *Venv);
+			UE_LOG(LogSwitchboardPlugin, Log, TEXT("Verify summary: %d"), Result.Summary);
+			UE_LOG(LogSwitchboardPlugin, Log, TEXT("Verify log: %s"), *Result.Log);
+
+			return Result;
+		});
+
+		VerifyPath = CurrentVenv;
+		VerifyResult = Future.Share();
+	}
+
+	return VerifyResult;
+}
+
+
+#if SB_LISTENER_AUTOLAUNCH
+bool FSwitchboardEditorModule::IsListenerAutolaunchEnabled(bool bForceRefreshCache /* = false */)
+{
+	if (bForceRefreshCache)
+	{
+		bCachedAutolaunchEnabled = GetListenerAutolaunchEnabled_Internal();
+	}
+
+	return bCachedAutolaunchEnabled;
+}
+
+
+bool FSwitchboardEditorModule::GetListenerAutolaunchEnabled_Internal() const
 {
 	const FString& ExistingCmd = GetListenerAutolaunchEntry().TrimEnd();
 	const FString& ConfigCmd = GetDefault<USwitchboardEditorSettings>()->GetListenerInvocation().TrimEnd();
@@ -74,17 +190,22 @@ bool FSwitchboardEditorModule::IsListenerAutolaunchEnabled() const
 
 bool FSwitchboardEditorModule::SetListenerAutolaunchEnabled(bool bEnabled)
 {
+	bool bSucceeded = false;
+
 	if (bEnabled)
 	{
 		const FString CommandLine = GetDefault<USwitchboardEditorSettings>()->GetListenerInvocation();
-		return SetListenerAutolaunchEntry(CommandLine);
+		bSucceeded = SetListenerAutolaunchEntry(CommandLine);
 	}
 	else
 	{
-		return RemoveListenerAutolaunchEntry();
+		bSucceeded = RemoveListenerAutolaunchEntry();
 	}
+
+	bCachedAutolaunchEnabled = GetListenerAutolaunchEnabled_Internal();
+	return bSucceeded;
 }
-#endif // #if PLATFORM_WINDOWS
+#endif // #if SB_LISTENER_AUTOLAUNCH
 
 
 void FSwitchboardEditorModule::OnEngineInitComplete()
@@ -96,9 +217,9 @@ void FSwitchboardEditorModule::OnEngineInitComplete()
 }
 
 
-bool FSwitchboardEditorModule::OnSettingsModified()
+bool FSwitchboardEditorModule::OnEditorSettingsModified()
 {
-#if PLATFORM_WINDOWS
+#if SB_LISTENER_AUTOLAUNCH
 	// If the existing entry's listener executable path matches our current engine / path,
 	// then we ensure the command line arguments also stay in sync with our settings.
 	const FString AutolaunchEntryExecutable = GetListenerAutolaunchEntryExecutable();
@@ -116,7 +237,7 @@ bool FSwitchboardEditorModule::OnSettingsModified()
 				*AutolaunchEntryExecutable, *ConfigExecutable);
 		}
 	}
-#endif // #if PLATFORM_WINDOWS
+#endif // #if SB_LISTENER_AUTOLAUNCH
 
 	return true;
 }
@@ -124,8 +245,8 @@ bool FSwitchboardEditorModule::OnSettingsModified()
 
 void FSwitchboardEditorModule::RunDefaultOSCListener()
 {
-	USwitchboardEditorSettings* SwitchboardEditorSettings = USwitchboardEditorSettings::GetSwitchboardEditorSettings();
-	FSoftObjectPath SwitchboardOSCListener = SwitchboardEditorSettings->SwitchboardOSCListener;
+	USwitchboardProjectSettings* SwitchboardProjectSettings = USwitchboardProjectSettings::GetSwitchboardProjectSettings();
+	FSoftObjectPath SwitchboardOSCListener = SwitchboardProjectSettings->SwitchboardOSCListener;
 	if (SwitchboardOSCListener.IsValid())
 	{
 		UObject* SwitchboardOSCListenerObject = SwitchboardOSCListener.TryLoad();
diff --git a/Engine/Plugins/VirtualProduction/Switchboard/Source/SwitchboardEditor/Private/SwitchboardEditorModule.h b/Engine/Plugins/VirtualProduction/Switchboard/Source/SwitchboardEditor/Private/SwitchboardEditorModule.h
--- a/Engine/Plugins/VirtualProduction/Switchboard/Source/SwitchboardEditor/Private/SwitchboardEditorModule.h
+++ b/Engine/Plugins/VirtualProduction/Switchboard/Source/SwitchboardEditor/Private/SwitchboardEditorModule.h
@@ -3,17 +3,34 @@
 #pragma once
 
 #include "CoreMinimal.h"
+#include "Async/Future.h"
 #include "Delegates/IDelegateInstance.h"
+#include "Misc/Paths.h"
 #include "Modules/ModuleInterface.h"
 #include "Modules/ModuleManager.h"
+#include "SwitchboardScriptInterop.h"
+
 
 
 DECLARE_LOG_CATEGORY_EXTERN(LogSwitchboardPlugin, Log, All);
 
 
+namespace UE::Switchboard::Private
+{
+	template <typename... TPaths>
+	FString ConcatPaths(FString BaseDir, TPaths... InPaths)
+	{
+		return (FPaths::ConvertRelativePathToFull(BaseDir) / ... / InPaths);
+	}
+} // namespace UE::Switchboard::Private
+
+
 class FSwitchboardEditorModule : public IModuleInterface
 {
 public:
+	static const FString& GetSbScriptsPath();
+	static const FString& GetSbThirdPartyPath();
+
 	/**
 	 * Singleton-like access to this module's interface.  This is just for convenience!
 	 * Beware of calling this during the shutdown phase, though.  Your module might have been unloaded already.
@@ -31,9 +48,17 @@ public:
 	virtual void ShutdownModule() override;
 	//~ End IModuleInterface
 
-#if PLATFORM_WINDOWS
-	/** Returns whether (this engine's) SwitchboardListener is configured to run automatically. */
-	bool IsListenerAutolaunchEnabled() const;
+	bool LaunchSwitchboard();
+	bool LaunchListener();
+
+	TSharedFuture<FSwitchboardVerifyResult> GetVerifyResult();
+
+#if SB_LISTENER_AUTOLAUNCH
+	/**
+	 * Returns whether (this engine's) SwitchboardListener is configured to run automatically.
+	 * Defaults to returning a cached value to avoid hitting the registry.
+	 */
+	bool IsListenerAutolaunchEnabled(bool bForceRefreshCache = false);
 
 	/** Enables or disables auto-run of SwitchboardListener. */
 	bool SetListenerAutolaunchEnabled(bool bEnabled);
@@ -41,10 +66,22 @@ public:
 
 private:
 	void OnEngineInitComplete();
-	bool OnSettingsModified();
+	bool OnEditorSettingsModified();
 
 	void RunDefaultOSCListener();
 
 private:
+#if SB_LISTENER_AUTOLAUNCH
+	bool GetListenerAutolaunchEnabled_Internal() const;
+#endif
+	bool RunProcess(const FString& InExe, const FString& InArgs);
+
 	FDelegateHandle DeferredStartDelegateHandle;
+
+#if SB_LISTENER_AUTOLAUNCH
+	bool bCachedAutolaunchEnabled;
+#endif
+
+	FString VerifyPath;
+	TSharedFuture<FSwitchboardVerifyResult> VerifyResult;
 };
diff --git a/Engine/Plugins/VirtualProduction/Switchboard/Source/SwitchboardEditor/Private/SwitchboardEditorSettings.cpp b/Engine/Plugins/VirtualProduction/Switchboard/Source/SwitchboardEditor/Private/SwitchboardEditorSettings.cpp
--- a/Engine/Plugins/VirtualProduction/Switchboard/Source/SwitchboardEditor/Private/SwitchboardEditorSettings.cpp
+++ b/Engine/Plugins/VirtualProduction/Switchboard/Source/SwitchboardEditor/Private/SwitchboardEditorSettings.cpp
@@ -1,57 +1,36 @@
 // Copyright Epic Games, Inc. All Rights Reserved.
 
 #include "SwitchboardEditorSettings.h"
+#include "SwitchboardEditorModule.h"
 #include "GenericPlatform/GenericPlatformProperties.h"
 #include "Misc/Paths.h"
 
 
-namespace UE::Switchboard::Private
+USwitchboardEditorSettings::USwitchboardEditorSettings()
 {
+	using namespace UE::Switchboard::Private;
 
-TTuple<FString, FString> PlatformAndExtension()
-{
-#if PLATFORM_WINDOWS
-	FString PlatformName = TEXT("Win64");
-	FString ExeExt = TEXT(".exe");
-#elif PLATFORM_LINUX
-	FString ExeExt = TEXT("");
-	FString PlatformName = TEXT("Linux");
-#endif
-	return TTuple<FString,FString>{MoveTemp(PlatformName), MoveTemp(ExeExt)};
+	VirtualEnvironmentPath.Path = ConcatPaths(FPaths::EngineDir(),
+		"Extras", "ThirdPartyNotUE", "SwitchboardThirdParty", "Python");
 }
 
-template <typename... TPaths>
-FString ConcatPaths(FString BaseDir, TPaths... InPaths)
-{
-	FString BaseConverted = FPaths::ConvertRelativePathToFull(BaseDir);
-	BaseConverted /= ( InPaths / ... );
-	return BaseConverted;
-}
 
-FString DefaultListenerPath()
+USwitchboardEditorSettings* USwitchboardEditorSettings::GetSwitchboardEditorSettings()
 {
-	const auto [PlatformName, ExeExt] = UE::Switchboard::Private::PlatformAndExtension();
-
-	return ConcatPaths(FPaths::EngineDir(),
-					   FString(TEXT("Binaries")),
-					   PlatformName,
-					   FString(TEXT("SwitchboardListener") + ExeExt));
+	return GetMutableDefault<USwitchboardEditorSettings>();
 }
 
-}
 
-USwitchboardEditorSettings::USwitchboardEditorSettings()
+FString USwitchboardEditorSettings::GetListenerPlatformPath() const
 {
-	using namespace UE::Switchboard::Private;
-	SwitchboardPath = {ConcatPaths(FPaths::EnginePluginsDir() + FString(TEXT("VirtualProduction")),
-								   FString(TEXT("Switchboard")),
-								   FString(TEXT("Source")),
-								   FString(TEXT("Switchboard")))};
-
-	ListenerPath = {DefaultListenerPath()};
+	FString Path = FPaths::ConvertRelativePathToFull(
+		FPlatformProcess::GenerateApplicationPath(TEXT("SwitchboardListener"), EBuildConfiguration::Development));
+	FPaths::MakePlatformFilename(Path);
+	return Path;
 }
 
-USwitchboardEditorSettings* USwitchboardEditorSettings::GetSwitchboardEditorSettings()
+
+FString USwitchboardEditorSettings::GetListenerInvocation() const
 {
-	return GetMutableDefault<USwitchboardEditorSettings>();
+	return FString::Printf(TEXT("\"%s\" %s"), *GetListenerPlatformPath(), *ListenerCommandlineArguments);
 }
diff --git a/Engine/Plugins/VirtualProduction/Switchboard/Source/SwitchboardEditor/Private/SwitchboardEditorStyle.cpp b/Engine/Plugins/VirtualProduction/Switchboard/Source/SwitchboardEditor/Private/SwitchboardEditorStyle.cpp
--- a/Engine/Plugins/VirtualProduction/Switchboard/Source/SwitchboardEditor/Private/SwitchboardEditorStyle.cpp
+++ b/Engine/Plugins/VirtualProduction/Switchboard/Source/SwitchboardEditor/Private/SwitchboardEditorStyle.cpp
@@ -4,8 +4,10 @@
 
 #include "EditorStyleSet.h"
 #include "Interfaces/IPluginManager.h"
-#include "Styling/SlateTypes.h"
 #include "Styling/CoreStyle.h"
+#include "Styling/SlateStyleMacros.h"
+#include "Styling/SlateTypes.h"
+#include "Styling/StyleColors.h"
 
 const FName FSwitchboardEditorStyle::NAME_SwitchboardBrush = "Img.Switchboard.Small";
 
@@ -13,12 +15,31 @@ const FName FSwitchboardEditorStyle::NAME_SwitchboardBrush = "Img.Switchboard.Sm
 FSwitchboardEditorStyle::FSwitchboardEditorStyle()
 	: FSlateStyleSet("SwitchboardEditorStyle")
 {
+	static const FVector2D Icon32x32(32.0f, 32.0f);
+
 	FString ContentDir = IPluginManager::Get().FindPlugin(TEXT("Switchboard"))->GetContentDir();
 	SetContentRoot(ContentDir);
 	SetCoreContentRoot(FPaths::EngineContentDir() / TEXT("Slate"));
 
 	Set(NAME_SwitchboardBrush, new FSlateImageBrush(RootToContentDir(TEXT("Icons/Switchboard_40x.png")), FVector2D(40.0f, 40.0f)));
-	
+
+	Set("Settings.RowBorder.Warning", new FSlateRoundedBoxBrush(FStyleColors::Transparent, 4.0f, FStyleColors::Warning, 1.0f));
+	Set("Settings.RowBorder.Nominal", new FSlateRoundedBoxBrush(FStyleColors::Transparent, 4.0f, FStyleColors::Success, 1.0f));
+
+	Set("Wizard.Background", new FSlateColorBrush(FStyleColors::Panel));
+
+	{
+		// These assets live in engine slate folders; change the root temporarily
+		const FString PreviousContentRoot = GetContentRootDir();
+		ON_SCOPE_EXIT{ SetContentRoot(PreviousContentRoot); };
+
+		const FString EngineSlateDir = FPaths::EngineContentDir() / TEXT("Slate");
+		SetContentRoot(EngineSlateDir);
+
+		Set("Settings.Icons.Warning", new IMAGE_BRUSH_SVG("Starship/Common/alert-triangle-large", Icon32x32, FStyleColors::Warning));
+		Set("Settings.Icons.Nominal", new IMAGE_BRUSH_SVG("Starship/Common/alert-triangle-large", Icon32x32, FStyleColors::Success));
+	}
+
 	FSlateStyleRegistry::RegisterSlateStyle(*this);
 }
 
diff --git a/Engine/Plugins/VirtualProduction/Switchboard/Source/SwitchboardEditor/Private/SwitchboardMenuEntry.cpp b/Engine/Plugins/VirtualProduction/Switchboard/Source/SwitchboardEditor/Private/SwitchboardMenuEntry.cpp
--- a/Engine/Plugins/VirtualProduction/Switchboard/Source/SwitchboardEditor/Private/SwitchboardMenuEntry.cpp
+++ b/Engine/Plugins/VirtualProduction/Switchboard/Source/SwitchboardEditor/Private/SwitchboardMenuEntry.cpp
@@ -91,16 +91,12 @@ struct FSwitchboardMenuEntryImpl
 				LOCTEXT("ListenerLaunchLabel", "Launch SwitchboardListener"),
 				LOCTEXT("ListenerLaunchTooltip", "Launches the SwitchboardListener with the settings from Editor Settings."),
 				FSlateIcon(),
-				FUIAction(FExecuteAction::CreateRaw(this, &FSwitchboardMenuEntryImpl::LaunchListener))
+				FUIAction(FExecuteAction::CreateRaw(this, &FSwitchboardMenuEntryImpl::OnLaunchListenerClicked))
 			);
 
-#if PLATFORM_WINDOWS
+#if SB_LISTENER_AUTOLAUNCH
 			MenuBuilder.AddMenuSeparator();
 
-			// The FGetActionCheckState delegate is called continually while the menu is open.
-			// Cache the state now as the combo is initially expanded.
-			CachedAutolaunchCheckState = GetAutolaunchCheckState();
-
 			MenuBuilder.AddMenuEntry(
 				LOCTEXT("ListenerAutolaunchLabel", "Launch Switchboard Listener on Login"),
 				LOCTEXT("ListenerAutolaunchTooltip", "Controls whether SwitchboardListener runs automatically when you log into Windows."),
@@ -108,21 +104,24 @@ struct FSwitchboardMenuEntryImpl
 				FUIAction(
 					FExecuteAction::CreateRaw(this, &FSwitchboardMenuEntryImpl::ToggleAutolaunch),
 					FCanExecuteAction(),
-					FGetActionCheckState::CreateLambda([&](){ return CachedAutolaunchCheckState; })
+					FGetActionCheckState::CreateLambda([&]()
+					{
+						return FSwitchboardEditorModule::Get().IsListenerAutolaunchEnabled() ? ECheckBoxState::Checked : ECheckBoxState::Unchecked;
+					})
 				),
 				NAME_None,
 				EUserInterfaceActionType::ToggleButton
 			);
-#endif
+#endif // #if SB_LISTENER_AUTOLAUNCH
 		}
 		MenuBuilder.EndSection();
 
 		return MenuBuilder.MakeWidget();
 	}
 
-	void LaunchListener()
+	void OnLaunchListenerClicked()
 	{
-		const FString ListenerPath = GetDefault<USwitchboardEditorSettings>()->ListenerPath.FilePath;
+		const FString ListenerPath = GetDefault<USwitchboardEditorSettings>()->GetListenerPlatformPath();
 		if (!FPaths::FileExists(ListenerPath))
 		{
 			const FString ErrorMsg = TEXT("Could not find SwitchboardListener! Make sure it was compiled.");
@@ -130,8 +129,7 @@ struct FSwitchboardMenuEntryImpl
 			return;
 		}
 
-		const FString ListenerArgs = GetDefault<USwitchboardEditorSettings>()->ListenerCommandlineArguments;
-		if (RunProcess(ListenerPath, ListenerArgs))
+		if (FSwitchboardEditorModule::Get().LaunchListener())
 		{
 			UE_LOG(LogSwitchboardPlugin, Display, TEXT("Successfully started SwitchboardListener"));
 		}
@@ -142,21 +140,16 @@ struct FSwitchboardMenuEntryImpl
 		}
 	}
 
-#if PLATFORM_WINDOWS
-	ECheckBoxState GetAutolaunchCheckState()
-	{
-		return FSwitchboardEditorModule::Get().IsListenerAutolaunchEnabled() ? ECheckBoxState::Checked : ECheckBoxState::Unchecked;
-	}
-
+#if SB_LISTENER_AUTOLAUNCH
 	void ToggleAutolaunch()
 	{
-		if (CachedAutolaunchCheckState == ECheckBoxState::Checked)
+		if (FSwitchboardEditorModule::Get().IsListenerAutolaunchEnabled())
 		{
 			FSwitchboardEditorModule::Get().SetListenerAutolaunchEnabled(false);
 		}
 		else
 		{
-			if (!FPaths::FileExists(GetDefault<USwitchboardEditorSettings>()->ListenerPath.FilePath))
+			if (!FPaths::FileExists(GetDefault<USwitchboardEditorSettings>()->GetListenerPlatformPath()))
 			{
 				const FString ErrorMsg = TEXT("Could not find SwitchboardListener! Make sure it has been compiled.");
 				FPlatformMisc::MessageBoxExt(EAppMsgType::Ok, *ErrorMsg, TEXT("Error enabling SwitchboardListener auto-launch"));
@@ -165,37 +158,14 @@ struct FSwitchboardMenuEntryImpl
 
 			FSwitchboardEditorModule::Get().SetListenerAutolaunchEnabled(true);
 		}
-
-		CachedAutolaunchCheckState = GetAutolaunchCheckState();
 	}
-#endif // #if PLATFORM_WINDOWS
+#endif // #if SB_LISTENER_AUTOLAUNCH
 
 	void OnLaunchSwitchboardClicked()
 	{
-		FString SwitchboardPath = GetDefault<USwitchboardEditorSettings>()->SwitchboardPath.Path;
-		if (SwitchboardPath.IsEmpty())
-		{
-			SwitchboardPath = FPaths::ConvertRelativePathToFull(FPaths::EnginePluginsDir() + FString(TEXT("VirtualProduction")));
-			SwitchboardPath /= FString(TEXT("Switchboard")) / FString(TEXT("Source")) / FString(TEXT("Switchboard"));
-		}
-
-#if PLATFORM_WINDOWS
-		FString Executable = SwitchboardPath / TEXT("switchboard.bat");
-#elif PLATFORM_LINUX
-		FString Executable = SwitchboardPath / TEXT("switchboard.sh");
-#endif
-		FString Args = GetDefault<USwitchboardEditorSettings>()->CommandlineArguments;
-
-		const FString PythonPath = GetDefault<USwitchboardEditorSettings>()->PythonInterpreterPath.FilePath;
-		if (!PythonPath.IsEmpty())
+		if (FSwitchboardEditorModule::Get().LaunchSwitchboard())
 		{
-			Executable = PythonPath;
-			Args.InsertAt(0, TEXT("-m switchboard "));
-		}
-
-		if (RunProcess(Executable, Args, SwitchboardPath))
-		{
-			UE_LOG(LogSwitchboardPlugin, Display, TEXT("Successfully started Switchboard from %s"), *SwitchboardPath);
+			UE_LOG(LogSwitchboardPlugin, Display, TEXT("Successfully started Switchboard"));
 		}
 		else
 		{
@@ -209,31 +179,11 @@ struct FSwitchboardMenuEntryImpl
 		return FSlateIcon(FSwitchboardEditorStyle::Get().GetStyleSetName(), FSwitchboardEditorStyle::NAME_SwitchboardBrush, FSwitchboardEditorStyle::NAME_SwitchboardBrush);
 	}
 
-	bool RunProcess(const FString& InExe, const FString& InArgs, const FString& InWorkingDirectory = TEXT(""))
-	{
-		const bool bLaunchDetached = false;
-		const bool bLaunchHidden = false;
-		const bool bLaunchReallyHidden = false;
-		const int32 PriorityModifier = 0;
-		const TCHAR* WorkingDirectory = InWorkingDirectory.IsEmpty() ? nullptr : *InWorkingDirectory;
-		uint32 PID = 0;
-		FProcHandle Handle = FPlatformProcess::CreateProc(*InExe, *InArgs, bLaunchDetached, bLaunchHidden, bLaunchReallyHidden, &PID, PriorityModifier, WorkingDirectory, nullptr);
-		int32 RetCode;
-		if (FPlatformProcess::GetProcReturnCode(Handle, &RetCode))
-		{
-			return RetCode == 0;
-		}
-		return FPlatformProcess::IsProcRunning(Handle);
-	}
-
 	bool IsToolbarButtonEnabled()
 	{
 		return true;
 	}
 
-private:
-	ECheckBoxState CachedAutolaunchCheckState;
-
 public:
 	static TUniquePtr<FSwitchboardMenuEntryImpl> Implementation;
 };
diff --git a/Engine/Plugins/VirtualProduction/Switchboard/Source/SwitchboardEditor/Private/SwitchboardProjectSettings.cpp b/Engine/Plugins/VirtualProduction/Switchboard/Source/SwitchboardEditor/Private/SwitchboardProjectSettings.cpp
new file mode 100644
--- /dev/null
+++ b/Engine/Plugins/VirtualProduction/Switchboard/Source/SwitchboardEditor/Private/SwitchboardProjectSettings.cpp
@@ -0,0 +1,14 @@
+// Copyright Epic Games, Inc. All Rights Reserved.
+
+#include "SwitchboardProjectSettings.h"
+
+
+USwitchboardProjectSettings::USwitchboardProjectSettings()
+{
+}
+
+
+USwitchboardProjectSettings* USwitchboardProjectSettings::GetSwitchboardProjectSettings()
+{
+	return GetMutableDefault<USwitchboardProjectSettings>();
+}
diff --git a/Engine/Plugins/VirtualProduction/Switchboard/Source/SwitchboardEditor/Private/SwitchboardScriptInterop.cpp b/Engine/Plugins/VirtualProduction/Switchboard/Source/SwitchboardEditor/Private/SwitchboardScriptInterop.cpp
new file mode 100644
--- /dev/null
+++ b/Engine/Plugins/VirtualProduction/Switchboard/Source/SwitchboardEditor/Private/SwitchboardScriptInterop.cpp
@@ -0,0 +1,229 @@
+// Copyright Epic Games, Inc. All Rights Reserved.
+
+#include "SwitchboardScriptInterop.h"
+#include "SwitchboardEditorModule.h"
+#include "Async/Async.h"
+#include "Serialization/JsonReader.h"
+#include "Serialization/JsonSerializer.h"
+
+
+TSharedRef<FSwitchboardUtilScript> FSwitchboardUtilScript::RunInstall(FStringView VenvPath)
+{
+	using namespace UE::Switchboard::Private;
+
+	const FString SbSetupPath = ConcatPaths(FSwitchboardEditorModule::Get().GetSbScriptsPath(), "sb_setup.py");
+
+	const FString Args = FString::Printf(TEXT("\"%s\" install --venv-dir=\"%.*s\""),
+		*SbSetupPath, VenvPath.Len(), VenvPath.GetData());
+
+	// We don't pass the venv to the constructor because we use the engine python to bootstrap.
+	TSharedRef<FSwitchboardUtilScript> Setup = MakeShared<FSwitchboardUtilScript>();
+	ensure(Setup->Run(Args));
+	return Setup;
+}
+
+
+FSwitchboardUtilScript::FSwitchboardUtilScript()
+{
+	using namespace UE::Switchboard::Private;
+
+	PythonExe = ConcatPaths(
+		FPaths::EngineDir(), "Binaries", "ThirdParty", "Python3",
+		FPlatformProcess::GetBinariesSubdirectory(),
+#if PLATFORM_WINDOWS
+		"pythonw.exe"
+#else
+		"bin", "python3"
+#endif
+	);
+}
+
+
+FSwitchboardUtilScript::FSwitchboardUtilScript(FStringView PythonVenvDir)
+{
+	using namespace UE::Switchboard::Private;
+
+	PythonExe = ConcatPaths(
+		FString(PythonVenvDir),
+#if PLATFORM_WINDOWS
+		"Scripts", "pythonw.exe"
+#else
+		"bin", "python3"
+#endif
+	);
+}
+
+
+FSwitchboardUtilScript::~FSwitchboardUtilScript()
+{
+	if (WritePipe || ReadPipe)
+	{
+		FPlatformProcess::ClosePipe(ReadPipe, WritePipe);
+	}
+
+	if (ProcHandle.IsValid())
+	{
+		FPlatformProcess::TerminateProc(ProcHandle);
+	}
+}
+
+
+bool FSwitchboardUtilScript::Run(const FString& Args)
+{
+	if (ProcHandle.IsValid())
+	{
+		// Not currently designed to Run() twice.
+		return ensure(false);
+	}
+
+	if (!FPaths::FileExists(PythonExe))
+	{
+		return false;
+	}
+
+	const bool bLaunchDetached = false;
+	const bool bLaunchHidden = true;
+	const bool bLaunchReallyHidden = true;
+	uint32 OutProcessId = 0;
+	const int32 PriorityModifier = 0;
+	const TCHAR* WorkingDirectory = nullptr;
+	FPlatformProcess::CreatePipe(ReadPipe, WritePipe);
+
+	ProcHandle = FPlatformProcess::CreateProc(*PythonExe, *Args, bLaunchDetached, bLaunchHidden, bLaunchReallyHidden, &OutProcessId, PriorityModifier, WorkingDirectory, WritePipe, ReadPipe);
+	return ProcHandle.IsValid();
+}
+
+
+TOptional<int32> FSwitchboardUtilScript::PollStdoutAndReturnCode()
+{
+	if (ReturnCode.IsSet())
+	{
+		// Process exited previously; nothing to do.
+		return ReturnCode;
+	}
+
+	if (!ProcHandle.IsValid())
+	{
+		return TOptional<int32>();
+	}
+
+	TArray<uint8> Output;
+	if (FPlatformProcess::ReadPipeToArray(ReadPipe, Output))
+	{
+		StdoutBuf.Append(Output);
+	}
+
+	int32 OutReturnCode;
+	if (FPlatformProcess::GetProcReturnCode(ProcHandle, &OutReturnCode))
+	{
+		ReturnCode.Emplace(OutReturnCode);
+
+		// Additional read necessary to ensure we've captured all output.
+		if (FPlatformProcess::ReadPipeToArray(ReadPipe, Output))
+		{
+			StdoutBuf.Append(Output);
+		}
+	}
+
+	return ReturnCode;
+}
+
+
+//static
+TFuture<FSwitchboardVerifyResult> FSwitchboardVerifyResult::RunVerify(FStringView VenvPath)
+{
+	using namespace UE::Switchboard::Private;
+
+	return Async(EAsyncExecution::Thread,
+		[
+			VenvPath = FString(VenvPath),
+			StartTimeSeconds = FPlatformTime::Seconds()
+		]() -> FSwitchboardVerifyResult
+		{
+			FSwitchboardVerifyResult Result;
+
+			FSwitchboardUtilScript Verify(VenvPath);
+
+			if (!FPaths::FileExists(Verify.PythonExe))
+			{
+				UE_LOG(LogSwitchboardPlugin, Warning, TEXT("RunVerify: Python interpreter not found: %s"), *Verify.PythonExe);
+				Result.Summary = FSwitchboardVerifyResult::ESummary::InterpreterMissing;
+				return Result;
+			}
+
+			const FString SbSetupPath = ConcatPaths(FSwitchboardEditorModule::Get().GetSbScriptsPath(), "sb_setup.py");
+			const FString Args = FString::Printf(TEXT("\"%s\" verify --output-json"), *SbSetupPath);
+			if (!Verify.Run(Args))
+			{
+				UE_LOG(LogSwitchboardPlugin, Warning, TEXT("RunVerify: Failed to launch: \"%s\" %s"),
+					*Verify.PythonExe, *Args);
+				Result.Summary = FSwitchboardVerifyResult::ESummary::LaunchFailed;
+				return Result;
+			}
+
+			TOptional<int32> ReturnCode;
+			while (!ReturnCode.IsSet())
+			{
+				ReturnCode = Verify.PollStdoutAndReturnCode();
+				if (!ReturnCode.IsSet())
+				{
+					const double VerifyTimeoutSeconds = 10.0f;
+					const double TimeElapsed = FPlatformTime::Seconds() - StartTimeSeconds;
+					if (TimeElapsed < VerifyTimeoutSeconds)
+					{
+						FPlatformProcess::Sleep(0.1f);
+					}
+					else
+					{
+						UE_LOG(LogSwitchboardPlugin, Warning, TEXT("RunVerify: Exceeded timeout"));
+						FPlatformProcess::TerminateProc(Verify.ProcHandle);
+						Result.Summary = FSwitchboardVerifyResult::ESummary::Timeout;
+						return Result;
+					}
+				}
+			}
+
+			FUTF8ToTCHAR Stdout(reinterpret_cast<const ANSICHAR*>(Verify.StdoutBuf.GetData()), Verify.StdoutBuf.Num());
+			FStringView StdoutSV(Stdout.Get(), Stdout.Length());
+
+			if (ReturnCode != 0)
+			{
+				UE_LOG(LogSwitchboardPlugin, Warning, TEXT("RunVerify: Returned non-zero exit status %d"), ReturnCode.GetValue());
+				Result.Summary = FSwitchboardVerifyResult::ESummary::UnexpectedExitStatus;
+				Result.Log = StdoutSV;
+				return Result;
+			}
+
+			Result.SetFromJson(StdoutSV);
+			return Result;
+		});
+}
+
+
+void FSwitchboardVerifyResult::SetFromJson(FStringView JsonSV)
+{
+	TSharedRef< TJsonReader<> > JsonReader = TJsonReaderFactory<>::Create(FString(JsonSV));
+	TSharedPtr< FJsonObject > JsonObject;
+	if (!FJsonSerializer::Deserialize(JsonReader, JsonObject) || !JsonObject.IsValid())
+	{
+		UE_LOG(LogSwitchboardPlugin, Warning, TEXT("FSwitchboardVerifyResult::SetFromJson: Couldn't deserialize"));
+		this->Summary = ESummary::InvalidOutput;
+		return;
+	}
+
+	if (!JsonObject->TryGetStringField("log", this->Log))
+	{
+		UE_LOG(LogSwitchboardPlugin, Error, TEXT("FSwitchboardVerifyResult: 'log' field missing"));
+		this->Log.Empty();
+	}
+
+	bool bSuccess = false;
+	if (!JsonObject->TryGetBoolField("success", bSuccess))
+	{
+		UE_LOG(LogSwitchboardPlugin, Error, TEXT("FSwitchboardVerifyResult: 'success' field missing"));
+		this->Summary = ESummary::InvalidOutput;
+		return;
+	}
+
+	this->Summary = bSuccess ? ESummary::Success : ESummary::SelfTestFailed;
+}
diff --git a/Engine/Plugins/VirtualProduction/Switchboard/Source/SwitchboardEditor/Private/SwitchboardScriptInterop.h b/Engine/Plugins/VirtualProduction/Switchboard/Source/SwitchboardEditor/Private/SwitchboardScriptInterop.h
new file mode 100644
--- /dev/null
+++ b/Engine/Plugins/VirtualProduction/Switchboard/Source/SwitchboardEditor/Private/SwitchboardScriptInterop.h
@@ -0,0 +1,54 @@
+// Copyright Epic Games, Inc. All Rights Reserved.
+
+#pragma once
+
+#include "CoreMinimal.h"
+#include "Async/Future.h"
+
+
+class FSwitchboardUtilScript
+{
+public:
+	static TSharedRef<FSwitchboardUtilScript> RunInstall(FStringView VenvPath);
+
+	/** Default constructor uses the engine python interpreter rather than a venv. */
+	FSwitchboardUtilScript();
+	explicit FSwitchboardUtilScript(FStringView PythonVenvDir);
+
+	~FSwitchboardUtilScript();
+
+	bool Run(const FString& Args);
+
+	TOptional<int32> PollStdoutAndReturnCode();
+
+public:
+	FString PythonExe;
+	FProcHandle ProcHandle;
+	void* WritePipe = nullptr;
+	void* ReadPipe = nullptr;
+	TArray<uint8> StdoutBuf;
+	TOptional<int32> ReturnCode;
+};
+
+
+struct FSwitchboardVerifyResult
+{
+	static TFuture<FSwitchboardVerifyResult> RunVerify(FStringView VenvPath);
+
+	enum class ESummary
+	{
+		InterpreterMissing,
+		LaunchFailed,
+		Timeout,
+		UnexpectedExitStatus,
+		InvalidOutput,
+		SelfTestFailed,
+
+		Success,
+	};
+
+	void SetFromJson(FStringView JsonSV);
+
+	ESummary Summary;
+	FString Log;
+};
diff --git a/Engine/Plugins/VirtualProduction/Switchboard/Source/SwitchboardEditor/Private/SwitchboardSettingsCustomization.cpp b/Engine/Plugins/VirtualProduction/Switchboard/Source/SwitchboardEditor/Private/SwitchboardSettingsCustomization.cpp
new file mode 100644
--- /dev/null
+++ b/Engine/Plugins/VirtualProduction/Switchboard/Source/SwitchboardEditor/Private/SwitchboardSettingsCustomization.cpp
@@ -0,0 +1,255 @@
+// Copyright Epic Games, Inc. All Rights Reserved.
+
+#include "SwitchboardSettingsCustomization.h"
+#include "SwitchboardEditorModule.h"
+#include "SwitchboardScriptInterop.h"
+#include "SwitchboardEditorSettings.h"
+#include "SwitchboardEditorStyle.h"
+#include "SwitchboardSetupWizard.h"
+#include "DetailCategoryBuilder.h"
+#include "DetailLayoutBuilder.h"
+#include "DetailWidgetRow.h"
+#include "ISettingsModule.h"
+#include "Misc/MessageDialog.h"
+#include "Widgets/Images/SImage.h"
+#include "Widgets/Input/SButton.h"
+#include "Widgets/Input/SCheckBox.h"
+#include "Widgets/Text/STextBlock.h"
+
+
+#define LOCTEXT_NAMESPACE "SwitchboardEditor"
+
+
+// static
+TSharedRef<IDetailCustomization> FSwitchboardEditorSettingsCustomization::MakeInstance()
+{
+	return MakeShared<FSwitchboardEditorSettingsCustomization>();
+}
+
+
+FSwitchboardEditorSettingsCustomization::ESwitchboardInstallState FSwitchboardEditorSettingsCustomization::GetInstallState() const
+{
+	TSharedFuture<FSwitchboardVerifyResult> VerifyResult = FSwitchboardEditorModule::Get().GetVerifyResult();
+
+	if (!VerifyResult.IsReady())
+	{
+		return ESwitchboardInstallState::VerifyInProgress;
+	}
+
+	if (VerifyResult.Get().Summary != FSwitchboardVerifyResult::ESummary::Success)
+	{
+		return ESwitchboardInstallState::NeedInstallOrRepair;
+	}
+
+	// TODO: Determine existence of shortcuts, indicate ESwitchboardInstallState::ShortcutsMissing
+
+	return ESwitchboardInstallState::Nominal;
+}
+
+
+FText FSwitchboardEditorSettingsCustomization::GetHealthRowHintText() const
+{
+	switch (GetInstallState())
+	{
+		case ESwitchboardInstallState::Nominal:
+			return LOCTEXT("ReinstallEnvHint", "Switchboard requires additional library/framework dependencies to launch, which are installed. If you encounter issues launching Switchboard, it may help to reinstall them.");
+		case ESwitchboardInstallState::ShortcutsMissing:
+			return LOCTEXT("AddShortcutsHint", "Add shortcuts to launch Switchboard from the Desktop and Start Menu.");
+		case ESwitchboardInstallState::NeedInstallOrRepair:
+		default:
+			return LOCTEXT("InstallEnvHint", "Switchboard requires additional library/framework dependencies to launch, please install them.");
+	}
+}
+
+
+FText FSwitchboardEditorSettingsCustomization::GetHealthRowButtonText() const
+{
+	switch (GetInstallState())
+	{
+		case ESwitchboardInstallState::Nominal:
+			return LOCTEXT("ReinstallEnvButton", "Reinstall Dependencies");
+		case ESwitchboardInstallState::ShortcutsMissing:
+			return LOCTEXT("AddShortcutsButton", "Add Shortcuts");
+		case ESwitchboardInstallState::NeedInstallOrRepair:
+		default:
+			return LOCTEXT("InstallEnvButton", "Install Dependencies");
+	}
+}
+
+
+const FSlateBrush* FSwitchboardEditorSettingsCustomization::GetHealthRowBorderBrush() const
+{
+	switch (GetInstallState())
+	{
+		case ESwitchboardInstallState::Nominal:
+			return FSwitchboardEditorStyle::Get().GetBrush("Settings.RowBorder.Nominal");
+		case ESwitchboardInstallState::ShortcutsMissing:
+		case ESwitchboardInstallState::NeedInstallOrRepair:
+		default:
+			return FSwitchboardEditorStyle::Get().GetBrush("Settings.RowBorder.Warning");
+	}
+}
+
+
+const FSlateBrush* FSwitchboardEditorSettingsCustomization::GetHealthRowIconBrush() const
+{
+	switch (GetInstallState())
+	{
+		case ESwitchboardInstallState::Nominal:
+			return FSwitchboardEditorStyle::Get().GetBrush("Settings.Icons.Nominal");
+		case ESwitchboardInstallState::ShortcutsMissing:
+		case ESwitchboardInstallState::NeedInstallOrRepair:
+		default:
+			return FSwitchboardEditorStyle::Get().GetBrush("Settings.Icons.Warning");
+	}
+}
+
+
+void FSwitchboardEditorSettingsCustomization::CustomizeDetails(IDetailLayoutBuilder& DetailLayout)
+{
+	IDetailCategoryBuilder& SwitchboardCategory = DetailLayout.EditCategory("Switchboard");
+
+	FText HealthFilterString = FText::GetEmpty();
+	TSharedPtr<IPropertyHandle> VenvProperty = DetailLayout.GetProperty(GET_MEMBER_NAME_CHECKED(USwitchboardEditorSettings, VirtualEnvironmentPath));
+	if (VenvProperty)
+	{
+		// Filter health row together with venv property row.
+		HealthFilterString = VenvProperty->GetPropertyDisplayName();
+	}
+
+	SwitchboardCategory.AddCustomRow(HealthFilterString, false)
+		.WholeRowWidget
+		[
+			SNew(SBorder)
+			.BorderImage(TAttribute<const FSlateBrush*>::CreateSP(this, &FSwitchboardEditorSettingsCustomization::GetHealthRowBorderBrush))
+			.Padding(FMargin(8.0f, 16.0f))
+			[
+				SNew(SHorizontalBox)
+				+ SHorizontalBox::Slot()
+				.AutoWidth()
+				.VAlign(VAlign_Center)
+				.Padding(FMargin(8.0f, 0.0f, 0.0f, 0.0f))
+				[
+					SNew(SImage)
+					.Image(TAttribute<const FSlateBrush*>::CreateSP(this, &FSwitchboardEditorSettingsCustomization::GetHealthRowIconBrush))
+				]
+				+ SHorizontalBox::Slot()
+				.FillWidth(1.0f)
+				.VAlign(VAlign_Center)
+				.Padding(16.0f, 0.0f)
+				[
+					SNew(STextBlock)
+					.AutoWrapText(true)
+					.Text(TAttribute<FText>::CreateSP(this, &FSwitchboardEditorSettingsCustomization::GetHealthRowHintText))
+				]
+				+ SHorizontalBox::Slot()
+				.AutoWidth()
+				.VAlign(VAlign_Center)
+				[
+					SNew(SButton)
+					.Text(TAttribute<FText>::CreateSP(this, &FSwitchboardEditorSettingsCustomization::GetHealthRowButtonText))
+					.OnClicked_Lambda([ThisWeak = TWeakPtr<FSwitchboardEditorSettingsCustomization>(SharedThis(this))]() {
+						SSwitchboardSetupWizard::EWizardPage StartPage = SSwitchboardSetupWizard::EWizardPage::Intro;
+						if (TSharedPtr<FSwitchboardEditorSettingsCustomization> This = ThisWeak.Pin())
+						{
+							if (This->GetInstallState() == ESwitchboardInstallState::ShortcutsMissing)
+							{
+								StartPage = SSwitchboardSetupWizard::EWizardPage::Shortcuts;
+							}
+						}
+						SSwitchboardSetupWizard::OpenWindow(StartPage);
+						return FReply::Handled();
+					})
+				]
+			]
+		];
+
+#if SB_LISTENER_AUTOLAUNCH
+	// This checkbox isn't backed by a property; it directly reflects the state
+	// of a Windows registry entry, and therefore only exists on Windows.
+	const FText& ListenerAutolaunchLabel = LOCTEXT("ListenerAutolaunchLabel", "Launch Switchboard Listener on Login");
+	IDetailCategoryBuilder& ListenerCategory = DetailLayout.EditCategory("Switchboard Listener");
+	FDetailWidgetRow& AutolaunchRow = ListenerCategory.AddCustomRow(ListenerAutolaunchLabel);
+	AutolaunchRow.NameContent()
+	[
+		SNew(STextBlock)
+		.Font(IDetailLayoutBuilder::GetDetailFont())
+		.Text(ListenerAutolaunchLabel)
+	];
+
+	AutolaunchRow.ValueContent()
+	[
+		SNew(SCheckBox)
+		.IsChecked_Lambda([]()
+		{
+			return FSwitchboardEditorModule::Get().IsListenerAutolaunchEnabled() ? ECheckBoxState::Checked : ECheckBoxState::Unchecked;
+		})
+		.OnCheckStateChanged_Lambda([](ECheckBoxState NewAutolaunchState)
+		{
+			const bool bShouldAutolaunch = NewAutolaunchState == ECheckBoxState::Checked;
+			if (bShouldAutolaunch && !FPaths::FileExists(GetDefault<USwitchboardEditorSettings>()->GetListenerPlatformPath()))
+			{
+				FMessageDialog::Open(EAppMsgType::Ok, LOCTEXT("ListenerMissing", "Unable to locate SwitchboardListener. Make sure it has been compiled."));
+				return;
+			}
+
+			FSwitchboardEditorModule::Get().SetListenerAutolaunchEnabled(bShouldAutolaunch);
+		})
+	];
+#endif // #if SB_LISTENER_AUTOLAUNCH
+}
+
+
+TSharedRef<IDetailCustomization> FSwitchboardProjectSettingsCustomization::MakeInstance()
+{
+	return MakeShared<FSwitchboardProjectSettingsCustomization>();
+}
+
+
+void FSwitchboardProjectSettingsCustomization::CustomizeDetails(IDetailLayoutBuilder& DetailLayout)
+{
+	const FText& HintLabel = LOCTEXT("JumpToEditorPrefsHint", "Configure Switchboard to launch on your computer via Editor Preferences.");
+
+	IDetailCategoryBuilder& EditorPrefsLinkCategory = DetailLayout.EditCategory("Configure Switchboard", LOCTEXT("JumpToEditorPrefsCategory", "Configure Switchboard"), ECategoryPriority::Uncommon);
+	EditorPrefsLinkCategory.AddCustomRow(HintLabel, false)
+		.WholeRowWidget
+		[
+			SNew(SBorder)
+			.BorderImage(FSwitchboardEditorStyle::Get().GetBrush("Settings.RowBorder.Warning"))
+			.Padding(FMargin(8.0f, 16.0f))
+			[
+				SNew(SHorizontalBox)
+				+ SHorizontalBox::Slot()
+				.AutoWidth()
+				.VAlign(VAlign_Center)
+				.Padding(8.0f, 0.0f, 0.0f, 0.0f)
+				[
+					SNew(SImage)
+					.Image(FEditorStyle::GetBrush("Icons.WarningWithColor.Large"))
+				]
+				+ SHorizontalBox::Slot()
+				.FillWidth(1.0f)
+				.VAlign(VAlign_Center)
+				.Padding(16.0f, 0.0f)
+				[
+					SNew(STextBlock)
+					.AutoWrapText(true)
+					.Text(HintLabel)
+				]
+				+ SHorizontalBox::Slot()
+				.AutoWidth()
+				.VAlign(VAlign_Center)
+				[
+					SNew(SButton)
+					.Text(LOCTEXT("JumpToEditorPrefsButtonText", "Take Me There"))
+					.OnClicked_Lambda([]() {
+						FModuleManager::LoadModuleChecked<ISettingsModule>("Settings").ShowViewer("Editor", "Plugins", "Switchboard");
+						return FReply::Handled();
+					})
+				]
+			]
+		];
+}
+
+
+#undef LOCTEXT_NAMESPACE
diff --git a/Engine/Plugins/VirtualProduction/Switchboard/Source/SwitchboardEditor/Private/SwitchboardSettingsCustomization.h b/Engine/Plugins/VirtualProduction/Switchboard/Source/SwitchboardEditor/Private/SwitchboardSettingsCustomization.h
new file mode 100644
--- /dev/null
+++ b/Engine/Plugins/VirtualProduction/Switchboard/Source/SwitchboardEditor/Private/SwitchboardSettingsCustomization.h
@@ -0,0 +1,45 @@
+// Copyright Epic Games, Inc. All Rights Reserved.
+
+#pragma once
+
+#include "CoreMinimal.h"
+#include "IDetailCustomization.h"
+
+
+struct FSlateBrush;
+
+
+class FSwitchboardEditorSettingsCustomization : public IDetailCustomization
+{
+public:
+	static TSharedRef<IDetailCustomization> MakeInstance();
+
+	enum class ESwitchboardInstallState
+	{
+		VerifyInProgress,
+
+		NeedInstallOrRepair,
+		ShortcutsMissing,
+		Nominal,
+	};
+
+	ESwitchboardInstallState GetInstallState() const;
+
+	FText GetHealthRowHintText() const;
+	FText GetHealthRowButtonText() const;
+	const FSlateBrush* GetHealthRowBorderBrush() const;
+	const FSlateBrush* GetHealthRowIconBrush() const;
+
+	//~ Begin IDetailCustomization Interface
+	void CustomizeDetails(IDetailLayoutBuilder& DetailLayout) override;
+};
+
+
+class FSwitchboardProjectSettingsCustomization : public IDetailCustomization
+{
+public:
+	static TSharedRef<IDetailCustomization> MakeInstance();
+
+	//~ Begin IDetailCustomization Interface
+	void CustomizeDetails(IDetailLayoutBuilder& DetailLayout) override;
+};
diff --git a/Engine/Plugins/VirtualProduction/Switchboard/Source/SwitchboardEditor/Private/SwitchboardSetupWizard.cpp b/Engine/Plugins/VirtualProduction/Switchboard/Source/SwitchboardEditor/Private/SwitchboardSetupWizard.cpp
new file mode 100644
--- /dev/null
+++ b/Engine/Plugins/VirtualProduction/Switchboard/Source/SwitchboardEditor/Private/SwitchboardSetupWizard.cpp
@@ -0,0 +1,594 @@
+// Copyright Epic Games, Inc. All Rights Reserved.
+
+#include "SwitchboardSetupWizard.h"
+#include "SwitchboardEditorModule.h"
+#include "SwitchboardEditorSettings.h"
+#include "SwitchboardEditorStyle.h"
+#include "SwitchboardScriptInterop.h"
+#include "Dialogs/CustomDialog.h"
+#include "Editor.h"
+#include "ISinglePropertyView.h"
+#include "SPrimaryButton.h"
+#include "Widgets/Images/SImage.h"
+#include "Widgets/Input/SCheckBox.h"
+#include "Widgets/Input/SEditableTextBox.h"
+#include "Widgets/Input/SMultiLineEditableTextBox.h"
+#include "Widgets/Layout/SSpacer.h"
+#include "Widgets/Layout/SWidgetSwitcher.h"
+#include "Widgets/Text/STextBlock.h"
+
+
+#define LOCTEXT_NAMESPACE "SwitchboardEditor"
+
+
+// static
+void SSwitchboardSetupWizard::OpenWindow(EWizardPage StartPage /* = EWizardPage::Intro */)
+{
+	TSharedRef<SWindow> ModalWindow = SNew(SWindow)
+		.SupportsMinimize(false)
+		.SupportsMaximize(false)
+		.ClientSize(FVector2D(700, 350));
+
+	const TSharedRef<SSwitchboardSetupWizard> Wizard = SNew(SSwitchboardSetupWizard, ModalWindow);
+	Wizard->SwitchToPage(StartPage);
+	ModalWindow->SetContent(Wizard);
+
+	GEditor->EditorAddModalWindow(ModalWindow);
+}
+
+
+void SSwitchboardSetupWizard::Construct(const FArguments& InArgs, const TSharedRef<SWindow>& InWindow)
+{
+	ModalWindow = InWindow;
+	ModalWindow->SetRequestDestroyWindowOverride(FRequestDestroyWindowOverride::CreateSP(SharedThis(this), &SSwitchboardSetupWizard::Handle_RequestDestroyWindow));
+
+	PageSwitcher = SNew(SWidgetSwitcher);
+
+	PageSwitcher->AddSlot(static_cast<int32>(EWizardPage::Intro))
+	[
+		Construct_Page_Intro()
+	];
+
+	PageSwitcher->AddSlot(static_cast<int32>(EWizardPage::InstallProgress))
+	[
+		Construct_Page_InstallProgress()
+	];
+
+	PageSwitcher->AddSlot(static_cast<int32>(EWizardPage::Shortcuts))
+	[
+		Construct_Page_Shortcuts()
+	];
+
+	PageSwitcher->AddSlot(static_cast<int32>(EWizardPage::Autolaunch))
+	[
+		Construct_Page_Autolaunch()
+	];
+
+	PageSwitcher->AddSlot(static_cast<int32>(EWizardPage::Complete))
+	[
+		Construct_Page_Complete()
+	];
+
+	ChildSlot
+	[
+		SNew(SOverlay)
+		+ SOverlay::Slot()
+		[
+			SNew(SBorder)
+			.BorderImage(FSwitchboardEditorStyle::Get().GetBrush("Wizard.Background"))
+		]
+		+ SOverlay::Slot()
+		.Padding(16.0f)
+		[
+			PageSwitcher.ToSharedRef()
+		]
+	];
+}
+
+
+TSharedRef<SWidget> SSwitchboardSetupWizard::Construct_Page_Intro()
+{
+	return SNew(SVerticalBox)
+		+ SVerticalBox::Slot()
+		.AutoHeight()
+		[
+			SNew(STextBlock)
+			.Text(LOCTEXT("Wizard_Intro_TextBlock", "Switchboard requires the following library/framework dependencies to launch:\n\n\t\u25CF  PySide2\n\n\t\u25CF  python-osc\n\n\t\u25CF  requests\n\n\t\u25CF  six\n\n\t\u25CF  cwRsync\n\n\nUnreal Engine will install the dependencies into the following directory:\n"))
+		]
+		+ SVerticalBox::Slot()
+		.AutoHeight()
+		[
+			SNew(SHorizontalBox)
+			+ SHorizontalBox::Slot()
+			[
+				SNew(SEditableTextBox)
+				.Text(FText::FromString(USwitchboardEditorSettings::GetSwitchboardEditorSettings()->VirtualEnvironmentPath.Path))
+				.OnTextCommitted_Lambda([](const FText& NewText, ETextCommit::Type CommitType) { USwitchboardEditorSettings::GetSwitchboardEditorSettings()->VirtualEnvironmentPath.Path = NewText.ToString(); })
+			]
+			+ SHorizontalBox::Slot()
+			.AutoWidth()
+			[
+				SNew(SButton)
+				.ButtonStyle(FEditorStyle::Get(), "SimpleButton")
+				.ContentPadding(FMargin(1, 0))
+				.Content()
+				[
+					SNew(SImage)
+					.Image(FAppStyle::Get().GetBrush("Icons.FolderClosed"))
+					.ColorAndOpacity(FSlateColor::UseForeground())
+				]
+			]
+		]
+		+ SVerticalBox::Slot()
+		[
+			SNew(SSpacer)
+		]
+		+ SVerticalBox::Slot()
+		.AutoHeight()
+		.HAlign(HAlign_Right)
+		[
+			SNew(SHorizontalBox)
+			+ SHorizontalBox::Slot()
+			.AutoWidth()
+			[
+				SNew(SPrimaryButton)
+				.Text(LOCTEXT("Wizard_Intro_InstallButton", "Install"))
+				.OnClicked(this, &SSwitchboardSetupWizard::Handle_Intro_InstallClicked)
+			]
+			+ SHorizontalBox::Slot()
+			.AutoWidth()
+			.Padding(16.0f, 0.0f, 0.0f, 0.0f)
+			[
+				SNew(SButton)
+				.Text(LOCTEXT("Wizard_Intro_CancelButton", "Cancel"))
+				.OnClicked_Lambda([ThisWeak = TWeakPtr<SSwitchboardSetupWizard>(SharedThis(this))]() {
+					if (TSharedPtr<SSwitchboardSetupWizard> This = ThisWeak.Pin())
+					{
+						This->ModalWindow->RequestDestroyWindow();
+					}
+					return FReply::Handled();
+				})
+			]
+		];
+}
+
+
+TSharedRef<SWidget> SSwitchboardSetupWizard::Construct_Page_InstallProgress()
+{
+	return SNew(SVerticalBox)
+		+ SVerticalBox::Slot()
+		.AutoHeight()
+		.HAlign(HAlign_Center)
+		.VAlign(VAlign_Center)
+		[
+			SNew(STextBlock)
+			.Text_Lambda([ThisWeak = TWeakPtr<SSwitchboardSetupWizard>(SharedThis(this))]() {
+				EInstallProgress Progress = EInstallProgress::Idle;
+				if (TSharedPtr<SSwitchboardSetupWizard> This = ThisWeak.Pin())
+				{
+					Progress = This->GetProgress();
+				}
+
+				switch (Progress)
+				{
+					case EInstallProgress::Running: return LOCTEXT("Wizard_Progress_InstallRunning", "Installation in progress...");
+					case EInstallProgress::Succeeded: return LOCTEXT("Wizard_Progress_InstallSucceeded", "Installation completed successfully!");
+					case EInstallProgress::Failed: return LOCTEXT("Wizard_Progress_InstallFailed", "Encountered an error during installation. Please refer to the log below for details.");
+					case EInstallProgress::Idle:
+					default:
+						return FText::GetEmpty();
+				}
+			})
+		]
+		+ SVerticalBox::Slot()
+		.Padding(0.0f, 16.0f)
+		.HAlign(HAlign_Fill)
+		.VAlign(VAlign_Fill)
+		[
+			SAssignNew(InstallOutputTextBox, SMultiLineEditableTextBox)
+			.IsReadOnly(true)
+			.Font(FEditorStyle::Get().GetWidgetStyle<FTextBlockStyle>("Log.Normal").Font)
+		]
+		+ SVerticalBox::Slot()
+		.AutoHeight()
+		.HAlign(HAlign_Right)
+		[
+			SNew(SHorizontalBox)
+			+ SHorizontalBox::Slot()
+			.AutoWidth()
+			[
+				SNew(SPrimaryButton)
+				.Text_Lambda([ThisWeak = TWeakPtr<SSwitchboardSetupWizard>(SharedThis(this))]() {
+					EInstallProgress Progress = EInstallProgress::Idle;
+					if (TSharedPtr<SSwitchboardSetupWizard> This = ThisWeak.Pin())
+					{
+						Progress = This->GetProgress();
+					}
+					return Progress != EInstallProgress::Failed ? LOCTEXT("Wizard_Progress_ContinueButton", "Continue") : LOCTEXT("Wizard_Progress_ExitButton", "Exit");
+				})
+				.IsEnabled_Lambda([ThisWeak = TWeakPtr<SSwitchboardSetupWizard>(SharedThis(this))]() {
+					EInstallProgress Progress = EInstallProgress::Idle;
+					if (TSharedPtr<SSwitchboardSetupWizard> This = ThisWeak.Pin())
+					{
+						Progress = This->GetProgress();
+					}
+					return Progress != EInstallProgress::Running;
+				})
+				.OnClicked(this, &SSwitchboardSetupWizard::Handle_InstallProgress_ContinueClicked)
+			]
+			+ SHorizontalBox::Slot()
+			.AutoWidth()
+			.Padding(16.0f, 0.0f, 0.0f, 0.0f)
+			[
+				SNew(SButton)
+				.Text(LOCTEXT("Wizard_Progress_CancelButton", "Cancel"))
+				.OnClicked_Lambda([ThisWeak = TWeakPtr<SSwitchboardSetupWizard>(SharedThis(this))]() {
+					if (TSharedPtr<SSwitchboardSetupWizard> This = ThisWeak.Pin())
+					{
+						// Handle_RequestDestroyWindow will prompt the user to confirm if interrupting installer
+						This->ModalWindow->RequestDestroyWindow();
+					}
+					return FReply::Handled();
+				})
+			]
+		];
+}
+
+
+TSharedRef<SWidget> SSwitchboardSetupWizard::Construct_Page_Shortcuts()
+{
+	return SNew(SVerticalBox)
+		+ SVerticalBox::Slot()
+		.AutoHeight()
+		[
+			SNew(STextBlock)
+			.Text(LOCTEXT("Wizard_Shortcuts_TextBlock", "Create shortcuts for Switchboard and Switchboard Listener:"))
+		]
+		+ SVerticalBox::Slot()
+		.FillHeight(1.0f)
+		[
+			SNew(SSpacer)
+		]
+		+ SVerticalBox::Slot()
+		.AutoHeight()
+		[
+			SNew(SHorizontalBox)
+			+ SHorizontalBox::Slot()
+			.AutoWidth()
+			.Padding(24.0f, 0.0f, 12.0f, 0.0f)
+			[
+				SNew(SCheckBox)
+			]
+			+ SHorizontalBox::Slot()
+			.AutoWidth()
+			[
+				SNew(STextBlock)
+				.Text(LOCTEXT("Wizard_Shortcuts_DesktopCheckboxLabel", "On the Desktop"))
+			]
+		]
+		+ SVerticalBox::Slot()
+		.AutoHeight()
+		.Padding(0.0f, 8.0f)
+		[
+			SNew(SHorizontalBox)
+			+ SHorizontalBox::Slot()
+			.AutoWidth()
+			.Padding(24.0f, 0.0f, 12.0f, 0.0f)
+			[
+				SNew(SCheckBox)
+			]
+			+ SHorizontalBox::Slot()
+			.AutoWidth()
+			[
+				SNew(STextBlock)
+				.Text(LOCTEXT("Wizard_Shortcuts_StartMenuCheckboxLabel", "In the Start Menu"))
+			]
+		]
+		+ SVerticalBox::Slot()
+		.FillHeight(3.0f)
+		[
+			SNew(SSpacer)
+		]
+		+ SVerticalBox::Slot()
+		.AutoHeight()
+		.HAlign(HAlign_Right)
+		[
+			SNew(SHorizontalBox)
+			+ SHorizontalBox::Slot()
+			.AutoWidth()
+			[
+				SNew(SPrimaryButton)
+				.Text(LOCTEXT("Wizard_Shortcuts_CreateButton", "Create"))
+				.OnClicked_Lambda([ThisWeak = TWeakPtr<SSwitchboardSetupWizard>(SharedThis(this))]() {
+					if (TSharedPtr<SSwitchboardSetupWizard> This = ThisWeak.Pin())
+					{
+						// TODO: if (...) CreateShortcuts();
+						This->SwitchToPage(EWizardPage::Autolaunch);
+					}
+					return FReply::Handled();
+				})
+			]
+			+ SHorizontalBox::Slot()
+			.AutoWidth()
+			.Padding(16.0f, 0.0f, 0.0f, 0.0f)
+			[
+				SNew(SButton)
+				.Text(LOCTEXT("Wizard_Shortcuts_SkipButton", "Cancel"))
+				.OnClicked_Lambda([ThisWeak = TWeakPtr<SSwitchboardSetupWizard>(SharedThis(this))]() {
+					if (TSharedPtr<SSwitchboardSetupWizard> This = ThisWeak.Pin())
+					{
+						This->SwitchToPage(EWizardPage::Autolaunch);
+					}
+					return FReply::Handled();
+				})
+			]
+		];
+}
+
+
+TSharedRef<SWidget> SSwitchboardSetupWizard::Construct_Page_Autolaunch()
+{
+	return SNew(SVerticalBox)
+		+ SVerticalBox::Slot()
+		.AutoHeight()
+		[
+			SNew(STextBlock)
+			.Text(LOCTEXT("Wizard_Autolaunch_TextBlock", "Would you like to automatically launch Switchboard Listener on login?"))
+		]
+		+ SVerticalBox::Slot()
+		.FillHeight(1.0f)
+		[
+			SNew(SSpacer)
+		]
+		+ SVerticalBox::Slot()
+		.AutoHeight()
+		[
+			SNew(SHorizontalBox)
+			+ SHorizontalBox::Slot()
+			.AutoWidth()
+			.Padding(24.0f, 0.0f, 12.0f, 0.0f)
+			[
+				SAssignNew(AutolaunchCheckbox, SCheckBox)
+#if !SB_LISTENER_AUTOLAUNCH
+				.IsEnabled(false)
+#endif
+			]
+			+ SHorizontalBox::Slot()
+			.AutoWidth()
+			[
+				SNew(STextBlock)
+#if SB_LISTENER_AUTOLAUNCH
+				.Text(LOCTEXT("Wizard_Autolaunch_CheckboxLabel", "Launch Switchboard Listener on Login"))
+#else
+				.IsEnabled(false)
+				.Text(LOCTEXT("Wizard_Autolaunch_CheckboxLabel_Unsupported", "Launch Switchboard Listener on Login\n(Not supported on this platform)"))
+#endif
+			]
+		]
+		+ SVerticalBox::Slot()
+		.FillHeight(3.0f)
+		[
+			SNew(SSpacer)
+		]
+		+ SVerticalBox::Slot()
+		.AutoHeight()
+		.HAlign(HAlign_Right)
+		[
+			SNew(SHorizontalBox)
+			+ SHorizontalBox::Slot()
+			.AutoWidth()
+			[
+				SNew(SPrimaryButton)
+				.Text(LOCTEXT("Wizard_Autolaunch_ContinueButton", "Continue"))
+				.OnClicked_Lambda([ThisWeak = TWeakPtr<SSwitchboardSetupWizard>(SharedThis(this))]() {
+					if (TSharedPtr<SSwitchboardSetupWizard> This = ThisWeak.Pin())
+					{
+#if SB_LISTENER_AUTOLAUNCH
+						FSwitchboardEditorModule::Get().SetListenerAutolaunchEnabled(This->AutolaunchCheckbox->IsChecked());
+#endif
+						This->SwitchToPage(EWizardPage::Complete);
+					}
+					return FReply::Handled();
+				})
+			]
+			+ SHorizontalBox::Slot()
+			.AutoWidth()
+			.Padding(16.0f, 0.0f, 0.0f, 0.0f)
+			[
+				SNew(SButton)
+				.Text(LOCTEXT("Wizard_Autolaunch_SkipButton", "Skip"))
+				.OnClicked_Lambda([ThisWeak = TWeakPtr<SSwitchboardSetupWizard>(SharedThis(this))]() {
+					if (TSharedPtr<SSwitchboardSetupWizard> This = ThisWeak.Pin())
+					{
+						This->SwitchToPage(EWizardPage::Complete);
+					}
+					return FReply::Handled();
+				})
+			]
+		];
+}
+
+
+TSharedRef<SWidget> SSwitchboardSetupWizard::Construct_Page_Complete()
+{
+		return SNew(SVerticalBox)
+		+ SVerticalBox::Slot()
+		.AutoHeight()
+		[
+			SNew(STextBlock)
+			.Text(LOCTEXT("Wizard_Complete_TextBlock", "Installation complete - Switchboard is now ready to launch!"))
+		]
+		+ SVerticalBox::Slot()
+		.FillHeight(1.0f)
+		[
+			SNew(SSpacer)
+		]
+		+ SVerticalBox::Slot()
+		.AutoHeight()
+		.HAlign(HAlign_Right)
+		[
+			SNew(SHorizontalBox)
+			+ SHorizontalBox::Slot()
+			.AutoWidth()
+			[
+				SNew(SPrimaryButton)
+				.Text(LOCTEXT("Wizard_Complete_LaunchButton", "Launch Switchboard"))
+				.OnClicked_Lambda([ThisWeak = TWeakPtr<SSwitchboardSetupWizard>(SharedThis(this))]() {
+					FSwitchboardEditorModule::Get().LaunchSwitchboard();
+					if (TSharedPtr<SSwitchboardSetupWizard> This = ThisWeak.Pin())
+					{
+						This->ModalWindow->RequestDestroyWindow();
+					}
+					return FReply::Handled();
+				})
+			]
+			+ SHorizontalBox::Slot()
+			.AutoWidth()
+			.Padding(16.0f, 0.0f, 0.0f, 0.0f)
+			[
+				SNew(SButton)
+				.Text(LOCTEXT("Wizard_Complete_CloseButton", "Close"))
+				.OnClicked_Lambda([ThisWeak = TWeakPtr<SSwitchboardSetupWizard>(SharedThis(this))]() {
+					if (TSharedPtr<SSwitchboardSetupWizard> This = ThisWeak.Pin())
+					{
+						This->ModalWindow->RequestDestroyWindow();
+					}
+					return FReply::Handled();
+				})
+			]
+		];
+}
+
+
+void SSwitchboardSetupWizard::Tick(const FGeometry& AllottedGeometry, const double InCurrentTime, const float InDeltaTime)
+{
+	SCompoundWidget::Tick(AllottedGeometry, InCurrentTime, InDeltaTime);
+
+	if (RunningInstall)
+	{
+		RunningInstall->PollStdoutAndReturnCode();
+
+		const int32 StdoutBufLen = RunningInstall->StdoutBuf.Num();
+		if (StdoutBufLen != LastStdoutNumBytes)
+		{
+			LastStdoutNumBytes = StdoutBufLen;
+			FUTF8ToTCHAR Stdout(reinterpret_cast<const ANSICHAR*>(RunningInstall->StdoutBuf.GetData()), StdoutBufLen);
+			InstallOutputTextBox->SetText(FText::FromStringView(FStringView(Stdout.Get(), Stdout.Length())));
+		}
+	}
+}
+
+
+void SSwitchboardSetupWizard::SwitchToPage(EWizardPage NewPage)
+{
+	PageSwitcher->SetActiveWidgetIndex(static_cast<int32>(NewPage));
+	ModalWindow->SetTitle(GetTitleForPage(NewPage));
+}
+
+
+FText SSwitchboardSetupWizard::GetTitleForPage(EWizardPage Page)
+{
+	switch (Page)
+	{
+		case EWizardPage::Intro: return LOCTEXT("Wizard_Intro_WindowTitle", "Install Switchboard Dependencies");
+		case EWizardPage::InstallProgress: return LOCTEXT("Wizard_InstallProgress_WindowTitle", "Install Switchboard Dependencies");
+		case EWizardPage::Shortcuts: return LOCTEXT("Wizard_Shortcuts_WindowTitle", "Add Switchboard Shortcuts");
+		case EWizardPage::Autolaunch: return LOCTEXT("Wizard_Autolaunch_WindowTitle", "Add Switchboard Shortcuts");
+		case EWizardPage::Complete: return LOCTEXT("Wizard_Complete_WindowTitle", "Switchboard Installation Complete");
+		default:
+			ensure(false);
+			return FText::GetEmpty();
+	}
+}
+
+
+SSwitchboardSetupWizard::EInstallProgress SSwitchboardSetupWizard::GetProgress() const
+{
+	if (!RunningInstall)
+	{
+		return EInstallProgress::Idle;
+	}
+	else if (!RunningInstall->ReturnCode.IsSet())
+	{
+		return EInstallProgress::Running;
+	}
+	else if (RunningInstall->ReturnCode == 0)
+	{
+		return EInstallProgress::Succeeded;
+	}
+	else // ReturnCode != 0
+	{
+		return EInstallProgress::Failed;
+	}
+}
+
+
+bool SSwitchboardSetupWizard::ConfirmAndInterruptInstall()
+{
+	TSharedRef<SCustomDialog> ConfirmDialog = SNew(SCustomDialog)
+		.Title(FText(LOCTEXT("Wizard_ConfirmInterruptInstall_WindowTitle", "Cancel Switchboard Installation?")))
+		.DialogContent(
+			SNew(STextBlock)
+			.Margin(FMargin(8.0f, 24.0f, 8.0f, 0.0f))
+			.Text(LOCTEXT("Wizard_ConfirmInterruptInstall_TextBlock", "Are you sure you want to cancel installation?\n\nSwitchboard requires these additional library/framework dependencies to launch."))
+		)
+		.Buttons({
+			SCustomDialog::FButton(LOCTEXT("Wizard_ConfirmInterruptInstall_CancelInstallButton", "Cancel Installation")),
+			SCustomDialog::FButton(LOCTEXT("Wizard_ConfirmInterruptInstall_ContinueInstallButton", "Continue Installation")),
+		});
+
+	const int32 ResultButtonIdx = ConfirmDialog->ShowModal();
+	if (ResultButtonIdx == 0 && GetProgress() == EInstallProgress::Running)
+	{
+		FPlatformProcess::TerminateProc(RunningInstall->ProcHandle, true);
+		return true;
+	}
+
+	return false;
+}
+
+
+void SSwitchboardSetupWizard::Handle_RequestDestroyWindow(const TSharedRef<SWindow>& InWindow)
+{
+	if (GetProgress() == EInstallProgress::Running)
+	{
+		if (!ConfirmAndInterruptInstall())
+		{
+			return;
+		}
+	}
+
+	FSlateApplicationBase::Get().RequestDestroyWindow(InWindow);
+}
+
+
+FReply SSwitchboardSetupWizard::Handle_Intro_InstallClicked()
+{
+	check(!RunningInstall.IsValid());
+
+	const FString VenvPath = USwitchboardEditorSettings::GetSwitchboardEditorSettings()->VirtualEnvironmentPath.Path;
+	RunningInstall = FSwitchboardUtilScript::RunInstall(VenvPath);
+
+	SwitchToPage(EWizardPage::InstallProgress);
+
+	return FReply::Handled();
+}
+
+
+FReply SSwitchboardSetupWizard::Handle_InstallProgress_ContinueClicked()
+{
+	if (GetProgress() == EInstallProgress::Succeeded)
+	{
+		// TODO: if (ShortcutsDontAlreadyExist())?
+		SwitchToPage(EWizardPage::Shortcuts);
+	}
+	else if (GetProgress() == EInstallProgress::Failed)
+	{
+		ModalWindow->RequestDestroyWindow();
+	}
+
+	return FReply::Handled();
+}
+
+
+#undef LOCTEXT_NAMESPACE
diff --git a/Engine/Plugins/VirtualProduction/Switchboard/Source/SwitchboardEditor/Private/SwitchboardSetupWizard.h b/Engine/Plugins/VirtualProduction/Switchboard/Source/SwitchboardEditor/Private/SwitchboardSetupWizard.h
new file mode 100644
--- /dev/null
+++ b/Engine/Plugins/VirtualProduction/Switchboard/Source/SwitchboardEditor/Private/SwitchboardSetupWizard.h
@@ -0,0 +1,72 @@
+// Copyright Epic Games, Inc. All Rights Reserved.
+
+#pragma once
+
+#include "CoreMinimal.h"
+#include "Widgets/DeclarativeSyntaxSupport.h"
+#include "Widgets/SCompoundWidget.h"
+
+
+class FSwitchboardUtilScript;
+class SCheckBox;
+class SMultiLineEditableTextBox;
+class SWidgetSwitcher;
+
+
+class SSwitchboardSetupWizard : public SCompoundWidget
+{
+	SLATE_BEGIN_ARGS(SSwitchboardSetupWizard) {}
+	SLATE_END_ARGS()
+
+public:
+	enum class EWizardPage
+	{
+		Intro,
+		InstallProgress,
+		Shortcuts,
+		Autolaunch,
+		Complete,
+	};
+
+	enum class EInstallProgress
+	{
+		Idle,
+		Running,
+		Succeeded,
+		Failed,
+	};
+
+	static void OpenWindow(EWizardPage StartPage = EWizardPage::Intro);
+
+	void Construct(const FArguments& InArgs, const TSharedRef<SWindow>& InWindow);
+
+protected:
+	TSharedRef<SWidget> Construct_Page_Intro();
+	TSharedRef<SWidget> Construct_Page_InstallProgress();
+	TSharedRef<SWidget> Construct_Page_Shortcuts();
+	TSharedRef<SWidget> Construct_Page_Autolaunch();
+	TSharedRef<SWidget> Construct_Page_Complete();
+
+	void Tick(const FGeometry& AllottedGeometry, const double InCurrentTime, const float InDeltaTime) override;
+
+	void SwitchToPage(EWizardPage NewPage);
+	FText GetTitleForPage(EWizardPage Page);
+	EInstallProgress GetProgress() const;
+
+	bool ConfirmAndInterruptInstall();
+
+	void Handle_RequestDestroyWindow(const TSharedRef<SWindow>& InWindow);
+
+	FReply Handle_Intro_InstallClicked();
+	FReply Handle_InstallProgress_ContinueClicked();
+
+protected:
+	TSharedPtr<SWindow> ModalWindow;
+	TSharedPtr<SWidgetSwitcher> PageSwitcher;
+
+	int32 LastStdoutNumBytes = 0;
+	TSharedPtr<SMultiLineEditableTextBox> InstallOutputTextBox;
+	TSharedPtr<SCheckBox> AutolaunchCheckbox;
+
+	TSharedPtr<FSwitchboardUtilScript> RunningInstall;
+};
diff --git a/Engine/Plugins/VirtualProduction/Switchboard/Source/SwitchboardEditor/Public/SwitchboardEditorSettings.h b/Engine/Plugins/VirtualProduction/Switchboard/Source/SwitchboardEditor/Public/SwitchboardEditorSettings.h
--- a/Engine/Plugins/VirtualProduction/Switchboard/Source/SwitchboardEditor/Public/SwitchboardEditorSettings.h
+++ b/Engine/Plugins/VirtualProduction/Switchboard/Source/SwitchboardEditor/Public/SwitchboardEditorSettings.h
@@ -11,7 +11,7 @@
 #include "SwitchboardEditorSettings.generated.h"
 
 
-UCLASS(config=Engine, BlueprintType)
+UCLASS(Config=Engine, BlueprintType)
 class SWITCHBOARDEDITOR_API USwitchboardEditorSettings : public UObject
 {
 	GENERATED_BODY()
@@ -19,46 +19,19 @@ class SWITCHBOARDEDITOR_API USwitchboardEditorSettings : public UObject
 	USwitchboardEditorSettings();
 
 public:
-	/** Switchboard installs its own python interpreter on its first launch. If you prefer
-	 * to use your own, specify the path to the python executable.
-	 */
+	/** Path to Switchboard's Python virtual environment, where third-party dependencies are installed. */
 	UPROPERTY(Config, EditAnywhere, Category="Switchboard")
-	FFilePath PythonInterpreterPath;
-
-	/** Path to Switchboard itself. */
-	UPROPERTY(Config, EditAnywhere, Category="Switchboard")
-	FDirectoryPath SwitchboardPath;
-
-	/** Arguments that should be passed to Switchboard. */
-	UPROPERTY(Config, EditAnywhere, Category="Switchboard")
-	FString CommandlineArguments;
-
-	/** Path to Switchboard Listener executable. */
-	UPROPERTY(Config, EditAnywhere, Category="Listener")
-	FFilePath ListenerPath;
+	FDirectoryPath VirtualEnvironmentPath;
 
 	/** Arguments that should be passed to the Switchboard Listener. */
-	UPROPERTY(Config, EditAnywhere, Category="Listener")
+	UPROPERTY(Config, EditAnywhere, Category="Switchboard Listener")
 	FString ListenerCommandlineArguments;
 
-	/** OSC Listener for Switchboard. An OSC server can be started on launch via VPUtilitiesEditor*/
-	UPROPERTY(config, EditAnywhere, Category = "OSC", meta = (DisplayName = "Default Switchboard OSC Listener", Tooltip = "The OSC listener for Switchboard. An OSC server can be started on launch via the Virtual Production Editor section of the Project Settings. Switchboard uses port 8000 by default, but this can be configured in your Switchboard config settings. "))
-	FSoftObjectPath SwitchboardOSCListener = FSoftObjectPath(TEXT("/Switchboard/OSCSwitchboard.OSCSwitchboard"));
-
-	/** Get Settings object for Switchboard*/
-	UFUNCTION(BlueprintPure, Category = "Switchboard")
+	/** Get Editor Settings object for Switchboard */
+	UFUNCTION(BlueprintPure, Category="Switchboard")
 	static USwitchboardEditorSettings* GetSwitchboardEditorSettings();
 
 public:
-	FString GetListenerPlatformPath() const
-	{
-		FString ListenerPlatformPath = ListenerPath.FilePath;
-		FPaths::MakePlatformFilename(ListenerPlatformPath);
-		return ListenerPlatformPath;
-	}
-
-	FString GetListenerInvocation() const
-	{
-		return FString::Printf(TEXT("\"%s\" %s"), *GetListenerPlatformPath(), *ListenerCommandlineArguments);
-	}
+	FString GetListenerPlatformPath() const;
+	FString GetListenerInvocation() const;
 };
diff --git a/Engine/Plugins/VirtualProduction/Switchboard/Source/SwitchboardEditor/Public/SwitchboardProjectSettings.h b/Engine/Plugins/VirtualProduction/Switchboard/Source/SwitchboardEditor/Public/SwitchboardProjectSettings.h
new file mode 100644
--- /dev/null
+++ b/Engine/Plugins/VirtualProduction/Switchboard/Source/SwitchboardEditor/Public/SwitchboardProjectSettings.h
@@ -0,0 +1,31 @@
+// Copyright Epic Games, Inc. All Rights Reserved.
+
+#pragma once
+
+#include "CoreMinimal.h"
+#include "Engine/EngineTypes.h"
+#include "Misc/Paths.h"
+#include "UObject/Object.h"
+#include "UObject/ObjectMacros.h"
+
+#include "SwitchboardProjectSettings.generated.h"
+
+
+UCLASS(Config=Game, BlueprintType)
+class SWITCHBOARDEDITOR_API USwitchboardProjectSettings : public UObject
+{
+	GENERATED_BODY()
+
+	USwitchboardProjectSettings();
+
+public:
+	/** OSC Listener for Switchboard. An OSC server can be started on launch via VPUtilitiesEditor */
+	UPROPERTY(Config, EditAnywhere, Category="OSC", meta=(
+	          DisplayName="Default Switchboard OSC Listener",
+	          Tooltip="The OSC listener for Switchboard. An OSC server can be started on launch via the Virtual Production Editor section of the Project Settings. Switchboard uses port 8000 by default, but this can be configured in your Switchboard config settings."))
+	FSoftObjectPath SwitchboardOSCListener = FSoftObjectPath(TEXT("/Switchboard/OSCSwitchboard.OSCSwitchboard"));
+
+	/** Get Project Settings object for Switchboard */
+	UFUNCTION(BlueprintPure, Category="Switchboard")
+	static USwitchboardProjectSettings* GetSwitchboardProjectSettings();
+};
diff --git a/Engine/Plugins/VirtualProduction/Switchboard/Source/SwitchboardEditor/SwitchboardEditor.build.cs b/Engine/Plugins/VirtualProduction/Switchboard/Source/SwitchboardEditor/SwitchboardEditor.build.cs
--- a/Engine/Plugins/VirtualProduction/Switchboard/Source/SwitchboardEditor/SwitchboardEditor.build.cs
+++ b/Engine/Plugins/VirtualProduction/Switchboard/Source/SwitchboardEditor/SwitchboardEditor.build.cs
@@ -14,13 +14,24 @@ public class SwitchboardEditor : ModuleRules
 				"Engine",
 				"EditorStyle",
 				"InputCore",
+				"Json",
 				"MessageLog",
 				"Projects",
 				"Settings",
 				"Slate",
 				"SlateCore",
 				"ToolMenus",
+				"ToolWidgets",
 				"UnrealEd",
 			});
+
+		if (Target.Platform == UnrealTargetPlatform.Win64)
+		{
+			PrivateDefinitions.Add("SB_LISTENER_AUTOLAUNCH=1");
+		}
+		else
+		{
+			PrivateDefinitions.Add("SB_LISTENER_AUTOLAUNCH=0");
+		}
 	}
 }
