Fix selected level being reset

diff --git a/Engine/Plugins/VirtualProduction/Switchboard/Source/Switchboard/switchboard/switchboard_dialog.py b/Engine/Plugins/VirtualProduction/Switchboard/Source/Switchboard/switchboard/switchboard_dialog.py
--- a/Engine/Plugins/VirtualProduction/Switchboard/Source/Switchboard/switchboard/switchboard_dialog.py
+++ b/Engine/Plugins/VirtualProduction/Switchboard/Source/Switchboard/switchboard/switchboard_dialog.py
@@ -1265,7 +1265,12 @@ class SwitchboardDialog(QtCore.QObject):
             CONFIG.save()
 
         if self.window.level_combo_box.currentText() != self._level:
-            self.window.level_combo_box.setCurrentText(self._level)
+            for index in range(self.window.level_combo_box.count()):
+                if self._get_level_from_combo_box(index) == self._level:
+                    self.window.level_combo_box.blockSignals(True)
+                    self.window.level_combo_box.setCurrentIndex(index)
+                    self.window.level_combo_box.blockSignals(False)
+                    break
 
     @property
     def project_changelist(self):
