Disable default focus for certain buttons

diff --git a/Engine/Plugins/VirtualProduction/Switchboard/Source/Switchboard/switchboard/config.py b/Engine/Plugins/VirtualProduction/Switchboard/Source/Switchboard/switchboard/config.py
--- a/Engine/Plugins/VirtualProduction/Switchboard/Source/Switchboard/switchboard/config.py
+++ b/Engine/Plugins/VirtualProduction/Switchboard/Source/Switchboard/switchboard/config.py
@@ -868,7 +868,7 @@ class ListSetting(Setting):
         migrate_data: typing.Callable[[any], None] = None
     ):
         '''
-        Create a new ArraySetting object.
+        Create a new ListSetting object.
 
         Args:
             attr_name      : Internal name.
@@ -939,6 +939,7 @@ class ListSetting(Setting):
         add_button.setSizePolicy(QtWidgets.QSizePolicy.Maximum, QtWidgets.QSizePolicy.Maximum)
         add_button.setMaximumWidth(12)
         add_button.setMaximumHeight(12)
+        add_button.setFocusPolicy(QtCore.Qt.FocusPolicy.NoFocus)
         add_button.pressed.connect(
             lambda override_device_name=override_device_name:
                 self._on_press_add(override_device_name)
@@ -951,6 +952,7 @@ class ListSetting(Setting):
         clear_button.setSizePolicy(QtWidgets.QSizePolicy.Maximum, QtWidgets.QSizePolicy.Maximum)
         clear_button.setMaximumWidth(12)
         clear_button.setMaximumHeight(12)
+        clear_button.setFocusPolicy(QtCore.Qt.FocusPolicy.NoFocus)
         clear_button.pressed.connect(
             lambda override_device_name=override_device_name:
             self._on_press_clear(override_device_name)
@@ -1545,6 +1547,7 @@ class LoggingSetting(Setting):
         edit_layout.addWidget(view)
 
         add_category_button = QtWidgets.QPushButton('Add Category')
+        add_category_button.setFocusPolicy(QtCore.Qt.FocusPolicy.NoFocus)
         edit_layout.addWidget(add_category_button)
 
         def on_add_category_button_clicked():
@@ -1566,6 +1569,7 @@ class LoggingSetting(Setting):
         add_category_button.clicked.connect(on_add_category_button_clicked)
 
         remove_category_button = QtWidgets.QPushButton('Remove Category')
+        remove_category_button.setFocusPolicy(QtCore.Qt.FocusPolicy.NoFocus)
         edit_layout.addWidget(remove_category_button)
 
         def on_remove_category_button_clicked():
