SwitchboardEditor: Reduce verbosity of log statements in Windows registry helpers. Reserve warning/error severity for use in response to user initiated actions, rather than inconsequential background queries.

diff --git a/Engine/Plugins/VirtualProduction/Switchboard/Source/SwitchboardEditor/Private/SwitchboardEditorModule.cpp b/Engine/Plugins/VirtualProduction/Switchboard/Source/SwitchboardEditor/Private/SwitchboardEditorModule.cpp
--- a/Engine/Plugins/VirtualProduction/Switchboard/Source/SwitchboardEditor/Private/SwitchboardEditorModule.cpp
+++ b/Engine/Plugins/VirtualProduction/Switchboard/Source/SwitchboardEditor/Private/SwitchboardEditorModule.cpp
@@ -265,7 +265,7 @@ HKEY OpenHkcuRunKey()
 	const LSTATUS OpenResult = RegOpenKeyEx(HKEY_CURRENT_USER, *WindowsRunRegKeyPath, 0, KEY_ALL_ACCESS, &HkcuRunKey);
 	if (OpenResult != ERROR_SUCCESS)
 	{
-		UE_LOG(LogSwitchboardPlugin, Error, TEXT("Error opening registry key %s (%08X)"), *WindowsRunRegKeyPath, OpenResult);
+		UE_LOG(LogSwitchboardPlugin, Log, TEXT("Error opening registry key %s (%08X)"), *WindowsRunRegKeyPath, OpenResult);
 		return nullptr;
 	}
 	return HkcuRunKey;
@@ -294,7 +294,7 @@ FString GetListenerAutolaunchEntry()
 	}
 	else if (SizeResult != ERROR_SUCCESS)
 	{
-		UE_LOG(LogSwitchboardPlugin, Error, TEXT("Error reading registry value %s:\"%s\" (%08X)"), *WindowsRunRegKeyPath, *ListenerAutolaunchRegValueName, SizeResult);
+		UE_LOG(LogSwitchboardPlugin, Log, TEXT("Error reading registry value %s:\"%s\" (%08X)"), *WindowsRunRegKeyPath, *ListenerAutolaunchRegValueName, SizeResult);
 		return FString();
 	}
 
@@ -306,12 +306,12 @@ FString GetListenerAutolaunchEntry()
 	const LSTATUS QueryResult = RegQueryValueEx(HkcuRunKey, *ListenerAutolaunchRegValueName, nullptr, &ValueType, reinterpret_cast<LPBYTE>(CharArray.GetData()), &ValueSizeBytes);
 	if (QueryResult != ERROR_SUCCESS)
 	{
-		UE_LOG(LogSwitchboardPlugin, Error, TEXT("Error reading registry value %s:\"%s\" (%08X)"), *WindowsRunRegKeyPath, *ListenerAutolaunchRegValueName, QueryResult);
+		UE_LOG(LogSwitchboardPlugin, Log, TEXT("Error reading registry value %s:\"%s\" (%08X)"), *WindowsRunRegKeyPath, *ListenerAutolaunchRegValueName, QueryResult);
 		return FString();
 	}
 	else if (ValueType != REG_SZ)
 	{
-		UE_LOG(LogSwitchboardPlugin, Error, TEXT("Registry value %s:\"%s\" has wrong type (%u)"), *WindowsRunRegKeyPath, *ListenerAutolaunchRegValueName, ValueType);
+		UE_LOG(LogSwitchboardPlugin, Log, TEXT("Registry value %s:\"%s\" has wrong type (%u)"), *WindowsRunRegKeyPath, *ListenerAutolaunchRegValueName, ValueType);
 		return FString();
 	}
 
