Switchboard onboarding: Immediately refresh venv verification after wizard concludes.

diff --git a/Engine/Plugins/VirtualProduction/Switchboard/Source/SwitchboardEditor/Private/SwitchboardEditorModule.cpp b/Engine/Plugins/VirtualProduction/Switchboard/Source/SwitchboardEditor/Private/SwitchboardEditorModule.cpp
--- a/Engine/Plugins/VirtualProduction/Switchboard/Source/SwitchboardEditor/Private/SwitchboardEditorModule.cpp
+++ b/Engine/Plugins/VirtualProduction/Switchboard/Source/SwitchboardEditor/Private/SwitchboardEditorModule.cpp
@@ -142,11 +142,11 @@ bool FSwitchboardEditorModule::RunProcess(const FString& InExe, const FString& I
 }
 
 
-TSharedFuture<FSwitchboardVerifyResult> FSwitchboardEditorModule::GetVerifyResult()
+TSharedFuture<FSwitchboardVerifyResult> FSwitchboardEditorModule::GetVerifyResult(bool bForceRefresh /* = false */)
 {
 	const FString CurrentVenv = GetDefault<USwitchboardEditorSettings>()->VirtualEnvironmentPath.Path;
 
-	if (!VerifyResult.IsValid() || VerifyPath != CurrentVenv)
+	if (bForceRefresh || !VerifyResult.IsValid() || VerifyPath != CurrentVenv)
 	{
 		UE_LOG(LogSwitchboardPlugin, Log, TEXT("Issuing verify for venv: %s"), *CurrentVenv);
 
diff --git a/Engine/Plugins/VirtualProduction/Switchboard/Source/SwitchboardEditor/Private/SwitchboardEditorModule.h b/Engine/Plugins/VirtualProduction/Switchboard/Source/SwitchboardEditor/Private/SwitchboardEditorModule.h
--- a/Engine/Plugins/VirtualProduction/Switchboard/Source/SwitchboardEditor/Private/SwitchboardEditorModule.h
+++ b/Engine/Plugins/VirtualProduction/Switchboard/Source/SwitchboardEditor/Private/SwitchboardEditorModule.h
@@ -51,7 +51,7 @@ public:
 	bool LaunchSwitchboard();
 	bool LaunchListener();
 
-	TSharedFuture<FSwitchboardVerifyResult> GetVerifyResult();
+	TSharedFuture<FSwitchboardVerifyResult> GetVerifyResult(bool bForceRefresh = false);
 
 #if SB_LISTENER_AUTOLAUNCH
 	/**
diff --git a/Engine/Plugins/VirtualProduction/Switchboard/Source/SwitchboardEditor/Private/SwitchboardSetupWizard.cpp b/Engine/Plugins/VirtualProduction/Switchboard/Source/SwitchboardEditor/Private/SwitchboardSetupWizard.cpp
--- a/Engine/Plugins/VirtualProduction/Switchboard/Source/SwitchboardEditor/Private/SwitchboardSetupWizard.cpp
+++ b/Engine/Plugins/VirtualProduction/Switchboard/Source/SwitchboardEditor/Private/SwitchboardSetupWizard.cpp
@@ -558,6 +558,10 @@ void SSwitchboardSetupWizard::Handle_RequestDestroyWindow(const TSharedRef<SWind
 		}
 	}
 
+	// Ensure we update the verification any time the wizard closes.
+	const bool bVerifyForceRefresh = true;
+	FSwitchboardEditorModule::Get().GetVerifyResult(bVerifyForceRefresh);
+
 	FSlateApplicationBase::Get().RequestDestroyWindow(InWindow);
 }
 
