Non-unique config names now have path next to them when pressed populate

diff --git a/Engine/Plugins/VirtualProduction/Switchboard/Source/Switchboard/switchboard/config.py b/Engine/Plugins/VirtualProduction/Switchboard/Source/Switchboard/switchboard/config.py
--- a/Engine/Plugins/VirtualProduction/Switchboard/Source/Switchboard/switchboard/config.py
+++ b/Engine/Plugins/VirtualProduction/Switchboard/Source/Switchboard/switchboard/config.py
@@ -2365,15 +2365,22 @@ class Config(object):
         del self._device_settings[(device_type, device_name)]
         self.save()
 
+    def shrink_path(self, path):
+        path_name = path.replace(self.get_project_dir(), '', 1)
+        path_name = path_name.replace(os.sep, '/')
+        return path_name
+        
+    def get_project_dir(self):
+        return os.path.dirname(self.UPROJECT_PATH.get_value().replace('"', ''))
+
     def maps(self):
         '''
         Returns a list of ful map paths in an Unreal Engine project such as [ "/Game/Maps/MapName" ].
         It will always start with Game and the slashes will always be "/" independent of the platform's separator.
         '''
-        project_dir = os.path.dirname(self.UPROJECT_PATH.get_value().replace('"', ''))
         maps_path = os.path.normpath(
             os.path.join(
-                project_dir,
+                self.get_project_dir(),
                 'Content',
                 self.MAPS_PATH.get_value()))
 
@@ -2384,11 +2391,10 @@ class Config(object):
                     continue
                 
                 map_name, _ = os.path.splitext(name)
-                path_name = path_to_map.replace(project_dir, '', 1)
                 # Ignore the fact that maps may exist in plugins - we only search game content
-                path_name = path_name.replace('Content', 'Game', 1)
+                path_name = path_to_map.replace('Content', 'Game', 1)
                 path_name = os.path.join(path_name, map_name)
-                path_name = path_name.replace(os.sep, '/')
+                path_name = self.shrink_path(path_name)
                 
                 if path_name not in maps:
                     maps.append(path_name)
diff --git a/Engine/Plugins/VirtualProduction/Switchboard/Source/Switchboard/switchboard/devices/ndisplay/plugin_ndisplay.py b/Engine/Plugins/VirtualProduction/Switchboard/Source/Switchboard/switchboard/devices/ndisplay/plugin_ndisplay.py
--- a/Engine/Plugins/VirtualProduction/Switchboard/Source/Switchboard/switchboard/devices/ndisplay/plugin_ndisplay.py
+++ b/Engine/Plugins/VirtualProduction/Switchboard/Source/Switchboard/switchboard/devices/ndisplay/plugin_ndisplay.py
@@ -16,6 +16,7 @@ from PySide2 import QtWidgets
 from switchboard import message_protocol
 from switchboard import switchboard_utils as sb_utils
 from switchboard import switchboard_widgets as sb_widgets
+from switchboard import switchboard_dialog as sb_dialog
 from switchboard.config import CONFIG, BoolSetting, FilePathSetting, \
     LoggingSetting, OptionSetting, Setting, StringSetting, SETTINGS, StringListSetting, \
     migrate_comma_separated_string_to_list
@@ -151,16 +152,20 @@ class AddnDisplayDialog(AddDeviceDialog):
             self.cbConfigs.setCurrentText(cfg_path)
             SETTINGS.LAST_BROWSED_PATH = os.path.dirname(cfg_path)
             SETTINGS.save()
-
+    
+    def generate_short_unique_config_name(config_path: str, file_name: str) -> str:
+        config_path = CONFIG.shrink_path(config_path)
+        return sb_dialog.SwitchboardDialog.filter_empty_abiguated_path(config_path, file_name)
+    
     def on_clicked_btnFindConfigs(self):
         ''' Finds and populates config combobox '''
+        
         # We will look for config files in the project's Content folder
         configs_path = os.path.normpath(
             os.path.join(
                 os.path.dirname(
                     CONFIG.UPROJECT_PATH.get_value().replace('"', '')),
                 'Content'))
-
         config_names = []
         config_paths = []
 
@@ -224,12 +229,13 @@ class AddnDisplayDialog(AddDeviceDialog):
                 # paths.
                 try:
                     asset = future.result()
-                    if asset['name'] not in config_names:
-                        config_names.append(asset['name'])
-                        config_paths.append(asset['path'])
+                    config_names.append(asset['name'])
+                    config_paths.append(asset['path'])
                 except Exception:
                     pass
 
+        config_names, _ = sb_dialog.SwitchboardDialog.generate_disambiguated_names(config_paths, AddnDisplayDialog.generate_short_unique_config_name)
+
         # close progress bar window
         progressDiag.close()
 
@@ -250,6 +256,8 @@ class AddnDisplayDialog(AddDeviceDialog):
 
         # update the combo box with the items.
         self.recall_config_itemDatas()
+        
+
 
     def devices_to_add(self):
         cfg_file = self.current_config_path()
diff --git a/Engine/Plugins/VirtualProduction/Switchboard/Source/Switchboard/switchboard/switchboard_dialog.py b/Engine/Plugins/VirtualProduction/Switchboard/Source/Switchboard/switchboard/switchboard_dialog.py
--- a/Engine/Plugins/VirtualProduction/Switchboard/Source/Switchboard/switchboard/switchboard_dialog.py
+++ b/Engine/Plugins/VirtualProduction/Switchboard/Source/Switchboard/switchboard/switchboard_dialog.py
@@ -1661,31 +1661,40 @@ class SwitchboardDialog(QtCore.QObject):
         if current_level and current_level in levels:
             self.level = current_level
 
+    def filter_empty_abiguated_path(path, file_name):
+        path = path.removesuffix(file_name)
+        if path == "/":
+            return f"{file_name}"
+        else:
+            return f"{file_name} ({path})"
 
-    def _update_level_list(self, level_combo_box: QtWidgets.QComboBox, level_path_list: List[str]):
-        def compare_file_names(path_a: str, path_b: str):
-            return -1 if path_a.lower() < path_b.lower() \
-                else 1 if path_a.lower() > path_b.lower() else 0
-        
-        def generate_short_map_path(map_path: str, file_name: str) -> str:
-            map_path = map_path.replace("/Game", "", 1)
-            map_path = map_path.removesuffix(file_name)
-            return f"{file_name} ({map_path})"
-            
+    def generate_short_map_path(path: str, file_name: str) -> str:
+            path = path.replace("/Game", "", 1)
+            return SwitchboardDialog.filter_empty_abiguated_path(path, file_name)
         
+    def generate_disambiguated_names(path_list, shortening_function):
         name_counts = {}
-        for map_path in level_path_list:
-            file_name = os.path.basename(map_path)
+        for path in path_list:
+            file_name = os.path.basename(path)
             name_counts[file_name] = name_counts.get(file_name, 0) + 1
             
         # Show only level name if unique and show path behind to disambiguate duplicates
         short_name_list = []
         short_name_to_path = {}
-        for map_path in level_path_list:
-            file_name = os.path.basename(map_path)
-            short_name = file_name if name_counts[file_name] == 1 else generate_short_map_path(map_path, file_name)
+        for path in path_list:
+            file_name = os.path.basename(path)
+            short_name = file_name if name_counts[file_name] == 1 else shortening_function(path, file_name)
             short_name_list.append(short_name)
-            short_name_to_path[short_name] = map_path
+            short_name_to_path[short_name] = path
+            
+        return short_name_list, short_name_to_path
+
+    def _update_level_list(self, level_combo_box: QtWidgets.QComboBox, level_path_list: List[str]):
+        def compare_file_names(path_a: str, path_b: str):
+            return -1 if path_a.lower() < path_b.lower() \
+                else 1 if path_a.lower() > path_b.lower() else 0
+        
+        short_name_list, short_name_to_path = SwitchboardDialog.generate_disambiguated_names(level_path_list, SwitchboardDialog.generate_short_map_path)
             
         from functools import cmp_to_key
         short_name_list = sorted(short_name_list, key=cmp_to_key(compare_file_names))
