Switchboard: Improved ExeFlags tooltip.

#rnx

diff --git a/Engine/Plugins/VirtualProduction/Switchboard/Source/Switchboard/switchboard/devices/ndisplay/ndisplay_monitor.py b/Engine/Plugins/VirtualProduction/Switchboard/Source/Switchboard/switchboard/devices/ndisplay/ndisplay_monitor.py
--- a/Engine/Plugins/VirtualProduction/Switchboard/Source/Switchboard/switchboard/devices/ndisplay/ndisplay_monitor.py
+++ b/Engine/Plugins/VirtualProduction/Switchboard/Source/Switchboard/switchboard/devices/ndisplay/ndisplay_monitor.py
@@ -63,9 +63,10 @@ class nDisplayMonitor(QAbstractTableModel):
                 'Whether nDisplay instance window is in focus. It is '
                 'recommended to be in focus.',
             'ExeFlags':
-                'It is recommended to disable fullscreen optimizations on the '
-                'unreal executable. Only available once the render node '
-                'process is running. Expects "DISABLEDXMAXIMIZEDWINDOWEDMODE"',
+                'It is recommended to disable fullscreen optimizations in the\n'
+                'unreal executable because it has been associated with tearing.\n'
+                'Only available once the render node process is running. \n'
+                '(Expected value is "DISABLEDXMAXIMIZEDWINDOWEDMODE")',
             'OsVer': 'Operating system version',
             'CpuUtilization':
                 'CPU utilization average. The number of overloaded cores (>'
