Switchboard: Return correct type when config parsing fails.

#rnx

diff --git a/Engine/Plugins/VirtualProduction/Switchboard/Source/Switchboard/switchboard/devices/ndisplay/plugin_ndisplay.py b/Engine/Plugins/VirtualProduction/Switchboard/Source/Switchboard/switchboard/devices/ndisplay/plugin_ndisplay.py
--- a/Engine/Plugins/VirtualProduction/Switchboard/Source/Switchboard/switchboard/devices/ndisplay/plugin_ndisplay.py
+++ b/Engine/Plugins/VirtualProduction/Switchboard/Source/Switchboard/switchboard/devices/ndisplay/plugin_ndisplay.py
@@ -1188,10 +1188,10 @@ class DevicenDisplay(DeviceUnreal):
         except Exception as e:
             LOGGER.error(
                 f'Error while parsing nDisplay config file "{cfg_file}": {e}')
-            return ([], None)
+            return DisplayConfig()
 
         LOGGER.error(f'Unknown nDisplay config file "{cfg_file}"')
-        return ([], None)
+        return DisplayConfig()
 
     def update_settings_controlled_by_config(self, cfg_file):
         '''
