Switchboard: Changed parse_config so it raises an exception when it fails to parse.

#rnx

diff --git a/Engine/Plugins/VirtualProduction/Switchboard/Source/Switchboard/switchboard/devices/ndisplay/plugin_ndisplay.py b/Engine/Plugins/VirtualProduction/Switchboard/Source/Switchboard/switchboard/devices/ndisplay/plugin_ndisplay.py
--- a/Engine/Plugins/VirtualProduction/Switchboard/Source/Switchboard/switchboard/devices/ndisplay/plugin_ndisplay.py
+++ b/Engine/Plugins/VirtualProduction/Switchboard/Source/Switchboard/switchboard/devices/ndisplay/plugin_ndisplay.py
@@ -1178,20 +1178,13 @@ class DevicenDisplay(DeviceUnreal):
         '''
         ext = os.path.splitext(cfg_file)[1].lower()
 
-        try:
-            if ext == '.ndisplay':
-                return cls.parse_config_json(cfg_file)
-
-            if ext == '.uasset':
-                return cls.parse_config_uasset(cfg_file)
-
-        except Exception as e:
-            LOGGER.error(
-                f'Error while parsing nDisplay config file "{cfg_file}": {e}')
-            return DisplayConfig()
+        if ext == '.ndisplay':
+            return cls.parse_config_json(cfg_file)
+        elif ext == '.uasset':
+            return cls.parse_config_uasset(cfg_file)
+        
+        raise Exception(f'Unknown config extension "{ext}"')
 
-        LOGGER.error(f'Unknown nDisplay config file "{cfg_file}"')
-        return DisplayConfig()
 
     def update_settings_controlled_by_config(self, cfg_file):
         '''
