SwitchboardEditor: Implement creation of Switchboard + Listener shortcuts during onboarding wizard. If shortcuts don't already exist, prompt user in editor preferences.

diff --git a/Engine/Plugins/VirtualProduction/Switchboard/Source/SwitchboardEditor/Private/SwitchboardEditorModule.cpp b/Engine/Plugins/VirtualProduction/Switchboard/Source/SwitchboardEditor/Private/SwitchboardEditorModule.cpp
--- a/Engine/Plugins/VirtualProduction/Switchboard/Source/SwitchboardEditor/Private/SwitchboardEditorModule.cpp
+++ b/Engine/Plugins/VirtualProduction/Switchboard/Source/SwitchboardEditor/Private/SwitchboardEditorModule.cpp
@@ -36,7 +36,7 @@ DEFINE_LOG_CATEGORY(LogSwitchboardPlugin);
 const FString& FSwitchboardEditorModule::GetSbScriptsPath()
 {
 	using namespace UE::Switchboard::Private;
-	static FString SbScriptsPath = ConcatPaths(FPaths::EnginePluginsDir(),
+	static const FString SbScriptsPath = ConcatPaths(FPaths::EnginePluginsDir(),
 		"VirtualProduction", "Switchboard", "Source", "Switchboard");
 	return SbScriptsPath;
 }
@@ -46,12 +46,25 @@ const FString& FSwitchboardEditorModule::GetSbScriptsPath()
 const FString& FSwitchboardEditorModule::GetSbThirdPartyPath()
 {
 	using namespace UE::Switchboard::Private;
-	static FString SbThirdPartyPath = ConcatPaths(FPaths::EngineDir(),
+	static const FString SbThirdPartyPath = ConcatPaths(FPaths::EngineDir(),
 		"Extras", "ThirdPartyNotUE", "SwitchboardThirdParty");
 	return SbThirdPartyPath;
 }
 
 
+// static
+const FString& FSwitchboardEditorModule::GetSbExePath()
+{
+#if PLATFORM_WINDOWS
+	static const FString ExePath = GetSbScriptsPath() / TEXT("switchboard.bat");
+#elif PLATFORM_LINUX
+	static const FString ExePath = GetSbScriptsPath() / TEXT("switchboard.sh");
+#endif
+
+	return ExePath;
+}
+
+
 void FSwitchboardEditorModule::StartupModule()
 {
 	FSwitchboardMenuEntry::Register();
@@ -106,16 +119,10 @@ void FSwitchboardEditorModule::ShutdownModule()
 
 bool FSwitchboardEditorModule::LaunchSwitchboard()
 {
-#if PLATFORM_WINDOWS
-	const FString ScriptPath = GetSbScriptsPath() / TEXT("switchboard.bat");
-#elif PLATFORM_LINUX
-	const FString ScriptPath = GetSbScriptsPath() / TEXT("switchboard.sh");
-#endif
-
 	const FString ScriptArgs = FString::Printf(TEXT("\"%s\""),
 		*GetDefault<USwitchboardEditorSettings>()->VirtualEnvironmentPath.Path);
 
-	return RunProcess(ScriptPath, ScriptArgs);
+	return RunProcess(GetSbExePath(), ScriptArgs);
 }
 
 
@@ -148,6 +155,10 @@ TSharedFuture<FSwitchboardVerifyResult> FSwitchboardEditorModule::GetVerifyResul
 
 	if (bForceRefresh || !VerifyResult.IsValid() || VerifyPath != CurrentVenv)
 	{
+#if SWITCHBOARD_SHORTCUTS
+		CachedShortcutCompares.Empty();
+#endif
+
 		UE_LOG(LogSwitchboardPlugin, Log, TEXT("Issuing verify for venv: %s"), *CurrentVenv);
 
 		TFuture<FSwitchboardVerifyResult> Future = FSwitchboardVerifyResult::RunVerify(CurrentVenv);
@@ -208,6 +219,47 @@ bool FSwitchboardEditorModule::SetListenerAutolaunchEnabled(bool bEnabled)
 #endif // #if SB_LISTENER_AUTOLAUNCH
 
 
+#if SWITCHBOARD_SHORTCUTS
+FSwitchboardEditorModule::EShortcutCompare FSwitchboardEditorModule::DoesShortcutExist(
+	EShortcutApp App,
+	EShortcutLocation Location,
+	bool bForceRefreshCache /* = false */
+)
+{
+	using namespace UE::Switchboard::Private::Shorcuts;
+
+	const TPair<EShortcutApp, EShortcutLocation> CacheKey{ App, Location };
+
+	if (!bForceRefreshCache)
+	{
+		if (const EShortcutCompare* CachedResult = CachedShortcutCompares.Find(CacheKey))
+		{
+			return *CachedResult;
+		}
+	}
+
+	const FShortcutParams ExpectedParams = BuildShortcutParams(App, Location);
+	const EShortcutCompare Result = CompareShortcut(ExpectedParams);
+	CachedShortcutCompares.FindOrAdd(CacheKey) = Result;
+	return Result;
+}
+
+
+bool FSwitchboardEditorModule::CreateOrUpdateShortcut(EShortcutApp App, EShortcutLocation Location)
+{
+	using namespace UE::Switchboard::Private::Shorcuts;
+
+	const FShortcutParams Params = BuildShortcutParams(App, Location);
+	const bool bResult = UE::Switchboard::Private::Shorcuts::CreateOrUpdateShortcut(Params);
+
+	const bool bForceUpdateCache = true;
+	(void)DoesShortcutExist(App, Location, bForceUpdateCache);
+
+	return bResult;
+}
+#endif // #if SWITCHBOARD_SHORTCUTS
+
+
 void FSwitchboardEditorModule::OnEngineInitComplete()
 {
 	FCoreDelegates::OnFEngineLoopInitComplete.Remove(DeferredStartDelegateHandle);
diff --git a/Engine/Plugins/VirtualProduction/Switchboard/Source/SwitchboardEditor/Private/SwitchboardEditorModule.h b/Engine/Plugins/VirtualProduction/Switchboard/Source/SwitchboardEditor/Private/SwitchboardEditorModule.h
--- a/Engine/Plugins/VirtualProduction/Switchboard/Source/SwitchboardEditor/Private/SwitchboardEditorModule.h
+++ b/Engine/Plugins/VirtualProduction/Switchboard/Source/SwitchboardEditor/Private/SwitchboardEditorModule.h
@@ -8,6 +8,7 @@
 #include "Misc/Paths.h"
 #include "Modules/ModuleInterface.h"
 #include "Modules/ModuleManager.h"
+#include "SwitchboardShortcuts.h"
 #include "SwitchboardScriptInterop.h"
 
 
@@ -30,6 +31,7 @@ class FSwitchboardEditorModule : public IModuleInterface
 public:
 	static const FString& GetSbScriptsPath();
 	static const FString& GetSbThirdPartyPath();
+	static const FString& GetSbExePath();
 
 	/**
 	 * Singleton-like access to this module's interface.  This is just for convenience!
@@ -64,6 +66,21 @@ public:
 	bool SetListenerAutolaunchEnabled(bool bEnabled);
 #endif
 
+#if SWITCHBOARD_SHORTCUTS
+	using EShortcutApp = UE::Switchboard::Private::Shorcuts::EShortcutApp;
+	using EShortcutLocation = UE::Switchboard::Private::Shorcuts::EShortcutLocation;
+	using EShortcutCompare = UE::Switchboard::Private::Shorcuts::EShortcutCompare;
+
+	/**
+	 * Returns whether shortcuts exist for (this engine's) Switchboard / Listener.
+	 * Defaults to returning a cached value to avoid hitting the filesystem.
+	 */
+	EShortcutCompare DoesShortcutExist(EShortcutApp App, EShortcutLocation Location, bool bForceRefreshCache = false);
+
+	/** Creates (or replaces) shortcuts for Switchboard / Listener. */
+	bool CreateOrUpdateShortcut(EShortcutApp App, EShortcutLocation Location);
+#endif
+
 private:
 	void OnEngineInitComplete();
 	bool OnEditorSettingsModified();
@@ -71,17 +88,20 @@ private:
 	void RunDefaultOSCListener();
 
 private:
-#if SB_LISTENER_AUTOLAUNCH
-	bool GetListenerAutolaunchEnabled_Internal() const;
-#endif
 	bool RunProcess(const FString& InExe, const FString& InArgs);
 
 	FDelegateHandle DeferredStartDelegateHandle;
 
+	FString VerifyPath;
+	TSharedFuture<FSwitchboardVerifyResult> VerifyResult;
+
 #if SB_LISTENER_AUTOLAUNCH
+	bool GetListenerAutolaunchEnabled_Internal() const;
+
 	bool bCachedAutolaunchEnabled;
 #endif
 
-	FString VerifyPath;
-	TSharedFuture<FSwitchboardVerifyResult> VerifyResult;
+#if SWITCHBOARD_SHORTCUTS
+	TMap< TPair<EShortcutApp, EShortcutLocation>, EShortcutCompare > CachedShortcutCompares;
+#endif
 };
diff --git a/Engine/Plugins/VirtualProduction/Switchboard/Source/SwitchboardEditor/Private/SwitchboardSettingsCustomization.cpp b/Engine/Plugins/VirtualProduction/Switchboard/Source/SwitchboardEditor/Private/SwitchboardSettingsCustomization.cpp
--- a/Engine/Plugins/VirtualProduction/Switchboard/Source/SwitchboardEditor/Private/SwitchboardSettingsCustomization.cpp
+++ b/Engine/Plugins/VirtualProduction/Switchboard/Source/SwitchboardEditor/Private/SwitchboardSettingsCustomization.cpp
@@ -29,7 +29,8 @@ TSharedRef<IDetailCustomization> FSwitchboardEditorSettingsCustomization::MakeIn
 
 FSwitchboardEditorSettingsCustomization::ESwitchboardInstallState FSwitchboardEditorSettingsCustomization::GetInstallState() const
 {
-	TSharedFuture<FSwitchboardVerifyResult> VerifyResult = FSwitchboardEditorModule::Get().GetVerifyResult();
+	FSwitchboardEditorModule& SbEditorModule = FSwitchboardEditorModule::Get();
+	TSharedFuture<FSwitchboardVerifyResult> VerifyResult = SbEditorModule.GetVerifyResult();
 
 	if (!VerifyResult.IsReady())
 	{
@@ -41,7 +42,20 @@ FSwitchboardEditorSettingsCustomization::ESwitchboardInstallState FSwitchboardEd
 		return ESwitchboardInstallState::NeedInstallOrRepair;
 	}
 
-	// TODO: Determine existence of shortcuts, indicate ESwitchboardInstallState::ShortcutsMissing
+#if SWITCHBOARD_SHORTCUTS
+	using namespace UE::Switchboard::Private::Shorcuts;
+
+	const bool bListenerDesktopShortcutExists = SbEditorModule.DoesShortcutExist(EShortcutApp::Listener, EShortcutLocation::Desktop) == EShortcutCompare::AlreadyExists;
+	const bool bListenerProgramsShortcutExists = SbEditorModule.DoesShortcutExist(EShortcutApp::Listener, EShortcutLocation::Programs) == EShortcutCompare::AlreadyExists;
+	const bool bAppDesktopShortcutExists = SbEditorModule.DoesShortcutExist(EShortcutApp::Switchboard, EShortcutLocation::Desktop) == EShortcutCompare::AlreadyExists;
+	const bool bAppProgramsShortcutExists = SbEditorModule.DoesShortcutExist(EShortcutApp::Switchboard, EShortcutLocation::Programs) == EShortcutCompare::AlreadyExists;
+
+	if (!bListenerDesktopShortcutExists || !bListenerProgramsShortcutExists
+		|| !bAppDesktopShortcutExists || !bAppProgramsShortcutExists)
+	{
+		return ESwitchboardInstallState::ShortcutsMissing;
+	}
+#endif // #if SWITCHBOARD_SHORTCUTS
 
 	return ESwitchboardInstallState::Nominal;
 }
diff --git a/Engine/Plugins/VirtualProduction/Switchboard/Source/SwitchboardEditor/Private/SwitchboardSetupWizard.cpp b/Engine/Plugins/VirtualProduction/Switchboard/Source/SwitchboardEditor/Private/SwitchboardSetupWizard.cpp
--- a/Engine/Plugins/VirtualProduction/Switchboard/Source/SwitchboardEditor/Private/SwitchboardSetupWizard.cpp
+++ b/Engine/Plugins/VirtualProduction/Switchboard/Source/SwitchboardEditor/Private/SwitchboardSetupWizard.cpp
@@ -257,13 +257,21 @@ TSharedRef<SWidget> SSwitchboardSetupWizard::Construct_Page_Shortcuts()
 			.AutoWidth()
 			.Padding(24.0f, 0.0f, 12.0f, 0.0f)
 			[
-				SNew(SCheckBox)
+				SAssignNew(DesktopShortcutCheckbox, SCheckBox)
+#if !SWITCHBOARD_SHORTCUTS
+				.IsEnabled(false)
+#endif
 			]
 			+ SHorizontalBox::Slot()
 			.AutoWidth()
 			[
 				SNew(STextBlock)
+#if SWITCHBOARD_SHORTCUTS
 				.Text(LOCTEXT("Wizard_Shortcuts_DesktopCheckboxLabel", "On the Desktop"))
+#else
+				.IsEnabled(false)
+				.Text(LOCTEXT("Wizard_Shortcuts_DesktopCheckboxLabel_Unsupported", "On the Desktop\n(Not supported on this platform)"))
+#endif
 			]
 		]
 		+ SVerticalBox::Slot()
@@ -275,13 +283,19 @@ TSharedRef<SWidget> SSwitchboardSetupWizard::Construct_Page_Shortcuts()
 			.AutoWidth()
 			.Padding(24.0f, 0.0f, 12.0f, 0.0f)
 			[
-				SNew(SCheckBox)
+				SAssignNew(ProgramsShortcutCheckbox, SCheckBox)
+#if !SWITCHBOARD_SHORTCUTS
+				.Visibility(EVisibility::Collapsed)
+#endif
 			]
 			+ SHorizontalBox::Slot()
 			.AutoWidth()
 			[
 				SNew(STextBlock)
 				.Text(LOCTEXT("Wizard_Shortcuts_StartMenuCheckboxLabel", "In the Start Menu"))
+#if !SWITCHBOARD_SHORTCUTS
+				.Visibility(EVisibility::Collapsed)
+#endif
 			]
 		]
 		+ SVerticalBox::Slot()
@@ -302,7 +316,19 @@ TSharedRef<SWidget> SSwitchboardSetupWizard::Construct_Page_Shortcuts()
 				.OnClicked_Lambda([ThisWeak = TWeakPtr<SSwitchboardSetupWizard>(SharedThis(this))]() {
 					if (TSharedPtr<SSwitchboardSetupWizard> This = ThisWeak.Pin())
 					{
-						// TODO: if (...) CreateShortcuts();
+#if SWITCHBOARD_SHORTCUTS
+						if (This->DesktopShortcutCheckbox->IsChecked())
+						{
+							FSwitchboardEditorModule::Get().CreateOrUpdateShortcut(FSwitchboardEditorModule::EShortcutApp::Switchboard, FSwitchboardEditorModule::EShortcutLocation::Desktop);
+							FSwitchboardEditorModule::Get().CreateOrUpdateShortcut(FSwitchboardEditorModule::EShortcutApp::Listener, FSwitchboardEditorModule::EShortcutLocation::Desktop);
+						}
+
+						if (This->ProgramsShortcutCheckbox->IsChecked())
+						{
+							FSwitchboardEditorModule::Get().CreateOrUpdateShortcut(FSwitchboardEditorModule::EShortcutApp::Switchboard, FSwitchboardEditorModule::EShortcutLocation::Programs);
+							FSwitchboardEditorModule::Get().CreateOrUpdateShortcut(FSwitchboardEditorModule::EShortcutApp::Listener, FSwitchboardEditorModule::EShortcutLocation::Programs);
+						}
+#endif
 						This->SwitchToPage(EWizardPage::Autolaunch);
 					}
 					return FReply::Handled();
@@ -583,7 +609,6 @@ FReply SSwitchboardSetupWizard::Handle_InstallProgress_ContinueClicked()
 {
 	if (GetProgress() == EInstallProgress::Succeeded)
 	{
-		// TODO: if (ShortcutsDontAlreadyExist())?
 		SwitchToPage(EWizardPage::Shortcuts);
 	}
 	else if (GetProgress() == EInstallProgress::Failed)
diff --git a/Engine/Plugins/VirtualProduction/Switchboard/Source/SwitchboardEditor/Private/SwitchboardSetupWizard.h b/Engine/Plugins/VirtualProduction/Switchboard/Source/SwitchboardEditor/Private/SwitchboardSetupWizard.h
--- a/Engine/Plugins/VirtualProduction/Switchboard/Source/SwitchboardEditor/Private/SwitchboardSetupWizard.h
+++ b/Engine/Plugins/VirtualProduction/Switchboard/Source/SwitchboardEditor/Private/SwitchboardSetupWizard.h
@@ -67,6 +67,8 @@ protected:
 	int32 LastStdoutNumBytes = 0;
 	TSharedPtr<SMultiLineEditableTextBox> InstallOutputTextBox;
 	TSharedPtr<SCheckBox> AutolaunchCheckbox;
+	TSharedPtr<SCheckBox> DesktopShortcutCheckbox;
+	TSharedPtr<SCheckBox> ProgramsShortcutCheckbox;
 
 	TSharedPtr<FSwitchboardUtilScript> RunningInstall;
 };
diff --git a/Engine/Plugins/VirtualProduction/Switchboard/Source/SwitchboardEditor/Private/SwitchboardShortcuts.cpp b/Engine/Plugins/VirtualProduction/Switchboard/Source/SwitchboardEditor/Private/SwitchboardShortcuts.cpp
new file mode 100644
--- /dev/null
+++ b/Engine/Plugins/VirtualProduction/Switchboard/Source/SwitchboardEditor/Private/SwitchboardShortcuts.cpp
@@ -0,0 +1,251 @@
+// Copyright Epic Games, Inc. All Rights Reserved.
+
+#include "SwitchboardShortcuts.h"
+#include "SwitchboardEditorModule.h"
+#include "SwitchboardEditorSettings.h"
+#include "Misc/EngineVersion.h"
+
+#if PLATFORM_WINDOWS
+#include "Windows/AllowWindowsPlatformTypes.h"
+#include <objbase.h>
+#include <ShlGuid.h>
+#include <ShlObj.h>
+#include "Windows/HideWindowsPlatformTypes.h"
+
+
+namespace UE::Switchboard::Private::Shorcuts {
+
+
+FShortcutParams BuildShortcutParams(EShortcutApp App, EShortcutLocation Location)
+{
+	const TCHAR* BaseName = nullptr;
+	switch (App)
+	{
+		case EShortcutApp::Switchboard: BaseName = TEXT("Switchboard"); break;
+		case EShortcutApp::Listener: BaseName = TEXT("Switchboard Listener"); break;
+		default:
+			checkNoEntry();
+			return FShortcutParams();
+	}
+
+	FShortcutParams Params;
+	Params.Location = Location;
+	Params.BaseName = FString::Printf(TEXT("%s (%d.%d)"), BaseName,
+		FEngineVersion::Current().GetMajor(), FEngineVersion::Current().GetMinor());
+
+	switch (App)
+	{
+		case EShortcutApp::Switchboard:
+		{
+			FString SbPath = FSwitchboardEditorModule::GetSbExePath();
+			FString VenvPath = GetDefault<USwitchboardEditorSettings>()->VirtualEnvironmentPath.Path;
+			FPaths::MakePlatformFilename(SbPath);
+			FPaths::MakePlatformFilename(VenvPath);
+			//Params.Target = FString::Printf(TEXT("\"%s\" \"%s\""), *SbPath, *VenvPath);
+			Params.Target = SbPath;
+			Params.Args = FString::Printf(TEXT("\"%s\""), *VenvPath);
+			break;
+		}
+		case EShortcutApp::Listener:
+		{
+			const FString ListenerPath = GetDefault<USwitchboardEditorSettings>()->GetListenerPlatformPath();
+			const FString& ListenerArgs = GetDefault<USwitchboardEditorSettings>()->ListenerCommandlineArguments;
+			//Params.Target = FString::Printf(TEXT("\"%s\" %s"), *ListenerPath, *ListenerArgs);
+			Params.Target = ListenerPath;
+			Params.Args = ListenerArgs;
+			break;
+		}
+		default:
+			checkNoEntry();
+			return FShortcutParams();
+	}
+
+	return Params;
+}
+
+bool CreateOrUpdateShortcut(const FShortcutParams& Params)
+{
+	if (!FWindowsPlatformMisc::CoInitialize())
+	{
+		return false;
+	}
+
+	ON_SCOPE_EXIT { FWindowsPlatformMisc::CoUninitialize(); };
+
+	IShellLink* ShellLink = nullptr;
+	HRESULT Hres = CoCreateInstance(CLSID_ShellLink, nullptr, CLSCTX_INPROC_SERVER, IID_IShellLink, reinterpret_cast<void**>(&ShellLink));
+	if (FAILED(Hres) || !ShellLink)
+	{
+		return false;
+	}
+
+	ON_SCOPE_EXIT { SAFE_RELEASE(ShellLink); };
+
+	IPersistFile* PersistFile = nullptr;
+	Hres = ShellLink->QueryInterface(&PersistFile);
+	if (FAILED(Hres) || !PersistFile)
+	{
+		return false;
+	}
+
+	ON_SCOPE_EXIT { SAFE_RELEASE(PersistFile); };
+
+	const FString SaveDir = GetShortcutLocationDir(Params.Location);
+	if (SaveDir.IsEmpty())
+	{
+		return false;
+	}
+
+	const FString SavePath = FString::Printf(TEXT("%s\\%s.lnk"), *SaveDir, *Params.BaseName);
+
+	ensure(SUCCEEDED(ShellLink->SetPath(*Params.Target)));
+	ensure(SUCCEEDED(ShellLink->SetArguments(*Params.Args)));
+	ensure(SUCCEEDED(ShellLink->SetDescription(*Params.Description)));
+	ensure(SUCCEEDED(ShellLink->SetIconLocation(*Params.IconPath, 0)));
+	if (!ensure(SUCCEEDED(PersistFile->Save(*SavePath, true))))
+	{
+		return false;
+	}
+
+	return true;
+}
+
+
+bool ReadShortcutParams(FShortcutParams& InOutParams)
+{
+	if (!FWindowsPlatformMisc::CoInitialize())
+	{
+		return false;
+	}
+
+	ON_SCOPE_EXIT { FWindowsPlatformMisc::CoUninitialize(); };
+
+	IShellLink* ShellLink = nullptr;
+	HRESULT Hres = CoCreateInstance(CLSID_ShellLink, nullptr, CLSCTX_INPROC_SERVER, IID_IShellLink, reinterpret_cast<void**>(&ShellLink));
+	if (FAILED(Hres) || !ShellLink)
+	{
+		return false;
+	}
+
+	ON_SCOPE_EXIT { SAFE_RELEASE(ShellLink); };
+
+	IPersistFile* PersistFile = nullptr;
+	Hres = ShellLink->QueryInterface(&PersistFile);
+	if (FAILED(Hres) || !PersistFile)
+	{
+		return false;
+	}
+
+	ON_SCOPE_EXIT { SAFE_RELEASE(PersistFile); };
+
+	const FString LoadDir = GetShortcutLocationDir(InOutParams.Location);
+	if (LoadDir.IsEmpty())
+	{
+		return false;
+	}
+
+	const FString LoadPath = FString::Printf(TEXT("%s\\%s.lnk"), *LoadDir, *InOutParams.BaseName);
+
+	if (!SUCCEEDED(PersistFile->Load(*LoadPath, STGM_READ)))
+	{
+		return false;
+	}
+
+	if (!SUCCEEDED(ShellLink->Resolve(nullptr, SLR_NO_UI | SLR_NOSEARCH)))
+	{
+		return false;
+	}
+
+	const int32 OutBufLen = 2048;
+	TArray<TCHAR> OutBuf;
+	OutBuf.SetNumUninitialized(OutBufLen);
+
+	if (ensure(SUCCEEDED(ShellLink->GetPath(OutBuf.GetData(), OutBufLen, nullptr, 0))))
+	{
+		InOutParams.Target = OutBuf.GetData();
+	}
+	else
+	{
+		InOutParams.Target.Empty();
+	}
+
+	if (ensure(SUCCEEDED(ShellLink->GetArguments(OutBuf.GetData(), OutBufLen))))
+	{
+		InOutParams.Args = OutBuf.GetData();
+	}
+	else
+	{
+		InOutParams.Args.Empty();
+	}
+
+	if (ensure(SUCCEEDED(ShellLink->GetDescription(OutBuf.GetData(), OutBufLen))))
+	{
+		InOutParams.Description = OutBuf.GetData();
+	}
+	else
+	{
+		InOutParams.Description.Empty();
+	}
+
+	int32 OutIconIndex = 0;
+	if (ensure(SUCCEEDED(ShellLink->GetIconLocation(OutBuf.GetData(), OutBufLen, &OutIconIndex))))
+	{
+		InOutParams.IconPath = OutBuf.GetData();
+	}
+	else
+	{
+		InOutParams.IconPath.Empty();
+	}
+
+	return true;
+}
+
+
+EShortcutCompare CompareShortcut(const FShortcutParams& InParams)
+{
+	FShortcutParams ExistingParams = InParams;
+	if (!ReadShortcutParams(ExistingParams))
+	{
+		return EShortcutCompare::Missing;
+	}
+
+	const bool bTargetDiffers = ExistingParams.Target != InParams.Target;
+	const bool bArgsDiffer = ExistingParams.Args != InParams.Args;
+	const bool bDescDiffers = ExistingParams.Description != InParams.Description;
+	const bool bIconDiffers = ExistingParams.IconPath != InParams.IconPath;
+
+	if (bTargetDiffers || bArgsDiffer || bDescDiffers || bIconDiffers)
+	{
+		return EShortcutCompare::Different;
+	}
+
+	return EShortcutCompare::AlreadyExists;
+}
+
+
+FString GetShortcutLocationDir(EShortcutLocation Location)
+{
+	const KNOWNFOLDERID* KnownFolder = nullptr;
+	switch (Location)
+	{
+		case EShortcutLocation::Desktop: KnownFolder = &FOLDERID_Desktop; break;
+		case EShortcutLocation::Programs: KnownFolder = &FOLDERID_Programs; break;
+		default: return FString();
+	}
+
+	TCHAR* OutKnownFolderPath = nullptr;
+	ON_SCOPE_EXIT{ CoTaskMemFree(OutKnownFolderPath); };
+	const HRESULT Hres = SHGetKnownFolderPath(*KnownFolder, 0, nullptr, &OutKnownFolderPath);
+	if (FAILED(Hres))
+	{
+		UE_LOG(LogSwitchboardPlugin, Log, TEXT("SHGetKnownFolderPath failed with HRESULT %08X"), Hres);
+		return FString();
+	}
+
+	return FString(OutKnownFolderPath);
+}
+
+
+} // namespace UE::Switchboard::Private::Shorcuts
+
+#endif // #if PLATFORM_WINDOWS
diff --git a/Engine/Plugins/VirtualProduction/Switchboard/Source/SwitchboardEditor/Private/SwitchboardShortcuts.h b/Engine/Plugins/VirtualProduction/Switchboard/Source/SwitchboardEditor/Private/SwitchboardShortcuts.h
new file mode 100644
--- /dev/null
+++ b/Engine/Plugins/VirtualProduction/Switchboard/Source/SwitchboardEditor/Private/SwitchboardShortcuts.h
@@ -0,0 +1,49 @@
+// Copyright Epic Games, Inc. All Rights Reserved.
+
+#pragma once
+
+#include "CoreMinimal.h"
+
+#if PLATFORM_WINDOWS
+
+namespace UE::Switchboard::Private::Shorcuts
+{
+	enum class EShortcutApp
+	{
+		Switchboard,
+		Listener,
+	};
+
+	enum class EShortcutLocation
+	{
+		Desktop,
+		Programs,
+	};
+
+	struct FShortcutParams
+	{
+		EShortcutLocation Location;
+		FString BaseName;
+
+		FString Target;
+		FString Args;
+		FString Description;
+		FString IconPath;
+	};
+
+	enum class EShortcutCompare
+	{
+		Missing,
+		Different,
+		AlreadyExists,
+	};
+
+	FShortcutParams BuildShortcutParams(EShortcutApp App, EShortcutLocation Location);
+	FString GetShortcutLocationDir(EShortcutLocation Location);
+
+	bool CreateOrUpdateShortcut(const FShortcutParams& Params);
+	bool ReadShortcutParams(FShortcutParams& InOutParams);
+	EShortcutCompare CompareShortcut(const FShortcutParams& Params);
+} // namespace UE::Switchboard::Private::Shorcuts
+
+#endif // #if PLATFORM_WINDOWS
\ No newline at end of file
diff --git a/Engine/Plugins/VirtualProduction/Switchboard/Source/SwitchboardEditor/SwitchboardEditor.build.cs b/Engine/Plugins/VirtualProduction/Switchboard/Source/SwitchboardEditor/SwitchboardEditor.build.cs
--- a/Engine/Plugins/VirtualProduction/Switchboard/Source/SwitchboardEditor/SwitchboardEditor.build.cs
+++ b/Engine/Plugins/VirtualProduction/Switchboard/Source/SwitchboardEditor/SwitchboardEditor.build.cs
@@ -28,10 +28,12 @@ public class SwitchboardEditor : ModuleRules
 		if (Target.Platform == UnrealTargetPlatform.Win64)
 		{
 			PrivateDefinitions.Add("SB_LISTENER_AUTOLAUNCH=1");
+			PrivateDefinitions.Add("SWITCHBOARD_SHORTCUTS=1");
 		}
 		else
 		{
 			PrivateDefinitions.Add("SB_LISTENER_AUTOLAUNCH=0");
+			PrivateDefinitions.Add("SWITCHBOARD_SHORTCUTS=0");
 		}
 	}
 }
