SwitchboardEditor: Launch directly into onboarding wizard from toolbar button if venv is not valid, and remove toolbar button if the user has previously completed installation and created shorcuts.

Other small fixes based on testing:
 - Always refresh the venv health when reopening the editor preferences.
 - When skipping to "create shortcuts" step of wizard, also exit wizard immediately afterward.
 - Default shortcut creation to checked (opt-out).
 - Correctly reflect prior autolaunch enabled/disabled state in wizard.
 - When launching SB via batch/shell script, ensure working directory is parent directory of Switchboard module, since some Switchboard code relies on this (causing incorrect fonts, possibly other issues).
 - Create cwrsync fstab file via sb_setup.py, instead of shipping it with the engine, so that we can also easily recreate it if necessary.

And minor cleanup:
 - Don't capture shared pointers in Slate lambdas.
 - Don't log warnings on "expected" verify issues, particularly ones we'll hit on build machines now that we trigger verify on startup.
 - Tidy UI_COMMANDs and Slate styles.

diff --git a/Engine/Plugins/VirtualProduction/Switchboard/Source/Switchboard/sb_setup.py b/Engine/Plugins/VirtualProduction/Switchboard/Source/Switchboard/sb_setup.py
--- a/Engine/Plugins/VirtualProduction/Switchboard/Source/Switchboard/sb_setup.py
+++ b/Engine/Plugins/VirtualProduction/Switchboard/Source/Switchboard/sb_setup.py
@@ -42,6 +42,16 @@ ENGINE_PATH = find_engine_ancestor_dir(SWITCHBOARD_PATH)
 THIRDPARTY_PATH = ENGINE_PATH / 'Extras/ThirdPartyNotUE'
 SB_THIRDPARTY_PATH = THIRDPARTY_PATH / 'SwitchboardThirdParty'
 
+CWRSYNC_SRC_DIR = THIRDPARTY_PATH / 'cwrsync'
+CWRSYNC_DEST_DIR = SB_THIRDPARTY_PATH / 'cwrsync'
+CWRSYNC_FSTAB_PATH = CWRSYNC_DEST_DIR / 'etc/fstab'
+CWRSYNC_FSTAB_CONTENTS = '''\
+# This is equivalent to the default, except for the addition of the "noacl"
+# option, which skips any attempt by Cygwin to store POSIX permissions via the
+# Windows host's NTFS access control lists.
+none /cygdrive cygdrive binary,posix=0,user,noacl 0 0
+'''
+
 
 class SbEnvBuilder(venv.EnvBuilder):
     def __init__(self, *args, **kwargs):
@@ -106,14 +116,16 @@ class SbSetup:
         if not sys.platform.startswith('win'):
             return
 
-        logging.debug('Copying cwrsync')
-
-        cwrsync_src_path = THIRDPARTY_PATH / 'cwrsync'
-        sb_cwrsync_dest_path = SB_THIRDPARTY_PATH / 'cwrsync'
-        if not (sb_cwrsync_dest_path / 'bin/rsync.exe').exists():
-            shutil.copytree(cwrsync_src_path, sb_cwrsync_dest_path,
+        if not (CWRSYNC_DEST_DIR / 'bin/rsync.exe').exists():
+            logging.info('Copying cwrsync')
+            shutil.copytree(CWRSYNC_SRC_DIR, CWRSYNC_DEST_DIR,
                             dirs_exist_ok=True)
 
+        if not (CWRSYNC_FSTAB_PATH).exists():
+            logging.info('Writing cwrsync fstab')
+            with open(CWRSYNC_FSTAB_PATH, 'wt') as fstab:
+                fstab.write(CWRSYNC_FSTAB_CONTENTS)
+
     @staticmethod
     def is_venv(dir: pathlib.Path) -> bool:
         ''' Conservative estimate of whether `dir` contains a Python venv. '''
diff --git a/Engine/Plugins/VirtualProduction/Switchboard/Source/Switchboard/switchboard.bat b/Engine/Plugins/VirtualProduction/Switchboard/Source/Switchboard/switchboard.bat
--- a/Engine/Plugins/VirtualProduction/Switchboard/Source/Switchboard/switchboard.bat
+++ b/Engine/Plugins/VirtualProduction/Switchboard/Source/Switchboard/switchboard.bat
@@ -54,7 +54,7 @@ goto:eof
 :start_sb
 
 set PYTHONPATH=%_switchboardDir%;%PYTHONPATH%
-start "Switchboard" "%_venvDir%\Scripts\pythonw.exe" -m switchboard
+start "Switchboard" /D "%_switchboardDir%" "%_venvDir%\Scripts\pythonw.exe" -m switchboard
 if %ERRORLEVEL% NEQ 0 (
     echo Failed to launch Switchboard^!
     if defined _bUsingDefaultVenv ( pause )
diff --git a/Engine/Plugins/VirtualProduction/Switchboard/Source/Switchboard/switchboard.sh b/Engine/Plugins/VirtualProduction/Switchboard/Source/Switchboard/switchboard.sh
--- a/Engine/Plugins/VirtualProduction/Switchboard/Source/Switchboard/switchboard.sh
+++ b/Engine/Plugins/VirtualProduction/Switchboard/Source/Switchboard/switchboard.sh
@@ -23,4 +23,4 @@ if [ ! -x "$_venvDir/bin/python3" ]; then
     fi
 fi
 
-PYTHONPATH="$_switchboardDir:$PYTHONPATH" "$_venvDir/bin/python3" -m switchboard
+(cd "$_switchboardDir" && PYTHONPATH="$_switchboardDir:$PYTHONPATH" "$_venvDir/bin/python3" -m switchboard)
diff --git a/Engine/Plugins/VirtualProduction/Switchboard/Source/SwitchboardEditor/Private/SwitchboardEditorModule.cpp b/Engine/Plugins/VirtualProduction/Switchboard/Source/SwitchboardEditor/Private/SwitchboardEditorModule.cpp
--- a/Engine/Plugins/VirtualProduction/Switchboard/Source/SwitchboardEditor/Private/SwitchboardEditorModule.cpp
+++ b/Engine/Plugins/VirtualProduction/Switchboard/Source/SwitchboardEditor/Private/SwitchboardEditorModule.cpp
@@ -5,6 +5,7 @@
 #include "SwitchboardMenuEntry.h"
 #include "SwitchboardProjectSettings.h"
 #include "SwitchboardSettingsCustomization.h"
+#include "Async/Async.h"
 #include "Editor.h"
 #include "EditorUtilitySubsystem.h"
 #include "ISettingsModule.h"
@@ -149,6 +150,37 @@ bool FSwitchboardEditorModule::RunProcess(const FString& InExe, const FString& I
 }
 
 
+FSwitchboardEditorModule::ESwitchboardInstallState FSwitchboardEditorModule::GetSwitchboardInstallState()
+{
+	TSharedFuture<FSwitchboardVerifyResult> Result = GetVerifyResult();
+
+	if (!Result.IsReady())
+	{
+		return ESwitchboardInstallState::VerifyInProgress;
+	}
+
+	if (Result.Get().Summary != FSwitchboardVerifyResult::ESummary::Success)
+	{
+		return ESwitchboardInstallState::NeedInstallOrRepair;
+	}
+
+#if SWITCHBOARD_SHORTCUTS
+	const bool bListenerDesktopShortcutExists = DoesShortcutExist(EShortcutApp::Listener, EShortcutLocation::Desktop) == EShortcutCompare::AlreadyExists;
+	const bool bListenerProgramsShortcutExists = DoesShortcutExist(EShortcutApp::Listener, EShortcutLocation::Programs) == EShortcutCompare::AlreadyExists;
+	const bool bAppDesktopShortcutExists = DoesShortcutExist(EShortcutApp::Switchboard, EShortcutLocation::Desktop) == EShortcutCompare::AlreadyExists;
+	const bool bAppProgramsShortcutExists = DoesShortcutExist(EShortcutApp::Switchboard, EShortcutLocation::Programs) == EShortcutCompare::AlreadyExists;
+
+	if (!bListenerDesktopShortcutExists || !bListenerProgramsShortcutExists
+		|| !bAppDesktopShortcutExists || !bAppProgramsShortcutExists)
+	{
+		return ESwitchboardInstallState::ShortcutsMissing;
+	}
+#endif // #if SWITCHBOARD_SHORTCUTS
+
+	return ESwitchboardInstallState::Nominal;
+}
+
+
 TSharedFuture<FSwitchboardVerifyResult> FSwitchboardEditorModule::GetVerifyResult(bool bForceRefresh /* = false */)
 {
 	const FString CurrentVenv = GetDefault<USwitchboardEditorSettings>()->VirtualEnvironmentPath.Path;
@@ -168,6 +200,17 @@ TSharedFuture<FSwitchboardVerifyResult> FSwitchboardEditorModule::GetVerifyResul
 			UE_LOG(LogSwitchboardPlugin, Log, TEXT("Verify summary: %d"), Result.Summary);
 			UE_LOG(LogSwitchboardPlugin, Log, TEXT("Verify log: %s"), *Result.Log);
 
+			Async(EAsyncExecution::TaskGraphMainThread, []() {
+				if (FSwitchboardEditorModule::Get().GetSwitchboardInstallState() == ESwitchboardInstallState::Nominal)
+				{
+					FSwitchboardMenuEntry::RemoveMenu();
+				}
+				else
+				{
+					FSwitchboardMenuEntry::AddMenu();
+				}
+			});
+
 			return Result;
 		});
 
@@ -266,6 +309,9 @@ void FSwitchboardEditorModule::OnEngineInitComplete()
 	DeferredStartDelegateHandle.Reset();
 
 	RunDefaultOSCListener();
+
+	// Populate initial verification results.
+	GetVerifyResult(true);
 }
 
 
diff --git a/Engine/Plugins/VirtualProduction/Switchboard/Source/SwitchboardEditor/Private/SwitchboardEditorModule.h b/Engine/Plugins/VirtualProduction/Switchboard/Source/SwitchboardEditor/Private/SwitchboardEditorModule.h
--- a/Engine/Plugins/VirtualProduction/Switchboard/Source/SwitchboardEditor/Private/SwitchboardEditorModule.h
+++ b/Engine/Plugins/VirtualProduction/Switchboard/Source/SwitchboardEditor/Private/SwitchboardEditorModule.h
@@ -53,6 +53,17 @@ public:
 	bool LaunchSwitchboard();
 	bool LaunchListener();
 
+	enum class ESwitchboardInstallState
+	{
+		VerifyInProgress,
+
+		NeedInstallOrRepair,
+		ShortcutsMissing,
+		Nominal,
+	};
+
+	ESwitchboardInstallState GetSwitchboardInstallState();
+
 	TSharedFuture<FSwitchboardVerifyResult> GetVerifyResult(bool bForceRefresh = false);
 
 #if SB_LISTENER_AUTOLAUNCH
diff --git a/Engine/Plugins/VirtualProduction/Switchboard/Source/SwitchboardEditor/Private/SwitchboardEditorStyle.cpp b/Engine/Plugins/VirtualProduction/Switchboard/Source/SwitchboardEditor/Private/SwitchboardEditorStyle.cpp
--- a/Engine/Plugins/VirtualProduction/Switchboard/Source/SwitchboardEditor/Private/SwitchboardEditorStyle.cpp
+++ b/Engine/Plugins/VirtualProduction/Switchboard/Source/SwitchboardEditor/Private/SwitchboardEditorStyle.cpp
@@ -9,8 +9,6 @@
 #include "Styling/SlateTypes.h"
 #include "Styling/StyleColors.h"
 
-const FName FSwitchboardEditorStyle::NAME_SwitchboardBrush = "Img.Switchboard.Small";
-
 
 FSwitchboardEditorStyle::FSwitchboardEditorStyle()
 	: FSlateStyleSet("SwitchboardEditorStyle")
@@ -21,24 +19,16 @@ FSwitchboardEditorStyle::FSwitchboardEditorStyle()
 	SetContentRoot(ContentDir);
 	SetCoreContentRoot(FPaths::EngineContentDir() / TEXT("Slate"));
 
-	Set(NAME_SwitchboardBrush, new FSlateImageBrush(RootToContentDir(TEXT("Icons/Switchboard_40x.png")), FVector2D(40.0f, 40.0f)));
+	// Key matches "(Context).(CommandId)" icon search used by our corresponding UI_COMMAND
+	Set("Switchboard.LaunchSwitchboard", new IMAGE_BRUSH("Icons/Switchboard_40x", FVector2D(40.0f, 40.0f)));
 
 	Set("Settings.RowBorder.Warning", new FSlateRoundedBoxBrush(FStyleColors::Transparent, 4.0f, FStyleColors::Warning, 1.0f));
 	Set("Settings.RowBorder.Nominal", new FSlateRoundedBoxBrush(FStyleColors::Transparent, 4.0f, FStyleColors::Success, 1.0f));
 
 	Set("Wizard.Background", new FSlateColorBrush(FStyleColors::Panel));
 
-	{
-		// These assets live in engine slate folders; change the root temporarily
-		const FString PreviousContentRoot = GetContentRootDir();
-		ON_SCOPE_EXIT{ SetContentRoot(PreviousContentRoot); };
-
-		const FString EngineSlateDir = FPaths::EngineContentDir() / TEXT("Slate");
-		SetContentRoot(EngineSlateDir);
-
-		Set("Settings.Icons.Warning", new IMAGE_BRUSH_SVG("Starship/Common/alert-triangle-large", Icon32x32, FStyleColors::Warning));
-		Set("Settings.Icons.Nominal", new IMAGE_BRUSH_SVG("Starship/Common/alert-triangle-large", Icon32x32, FStyleColors::Success));
-	}
+	Set("Settings.Icons.Warning", new CORE_IMAGE_BRUSH_SVG("Starship/Common/alert-triangle-large", Icon32x32, FStyleColors::Warning));
+	Set("Settings.Icons.Nominal", new CORE_IMAGE_BRUSH_SVG("Starship/Common/alert-triangle-large", Icon32x32, FStyleColors::Success));
 
 	FSlateStyleRegistry::RegisterSlateStyle(*this);
 }
diff --git a/Engine/Plugins/VirtualProduction/Switchboard/Source/SwitchboardEditor/Private/SwitchboardEditorStyle.h b/Engine/Plugins/VirtualProduction/Switchboard/Source/SwitchboardEditor/Private/SwitchboardEditorStyle.h
--- a/Engine/Plugins/VirtualProduction/Switchboard/Source/SwitchboardEditor/Private/SwitchboardEditorStyle.h
+++ b/Engine/Plugins/VirtualProduction/Switchboard/Source/SwitchboardEditor/Private/SwitchboardEditorStyle.h
@@ -13,8 +13,6 @@ class SWITCHBOARDEDITOR_API FSwitchboardEditorStyle : public FSlateStyleSet
 public:
 	static FSwitchboardEditorStyle& Get();
 
-	static const FName NAME_SwitchboardBrush;
-
 private:
 	FSwitchboardEditorStyle();
 	~FSwitchboardEditorStyle();
diff --git a/Engine/Plugins/VirtualProduction/Switchboard/Source/SwitchboardEditor/Private/SwitchboardMenuEntry.cpp b/Engine/Plugins/VirtualProduction/Switchboard/Source/SwitchboardEditor/Private/SwitchboardMenuEntry.cpp
--- a/Engine/Plugins/VirtualProduction/Switchboard/Source/SwitchboardEditor/Private/SwitchboardMenuEntry.cpp
+++ b/Engine/Plugins/VirtualProduction/Switchboard/Source/SwitchboardEditor/Private/SwitchboardMenuEntry.cpp
@@ -15,6 +15,7 @@
 #include "SwitchboardEditorModule.h"
 #include "SwitchboardEditorSettings.h"
 #include "SwitchboardEditorStyle.h"
+#include "SwitchboardSetupWizard.h"
 #include "ToolMenus.h"
 
 #define LOCTEXT_NAMESPACE "SwitchboardEditor"
@@ -29,10 +30,12 @@ public:
 
 	virtual void RegisterCommands() override
 	{
-		UI_COMMAND(TriggerToolbarButtonCmd, "Launch Switchboard", "Launch Switchboard", EUserInterfaceActionType::Button, FInputChord());
+		UI_COMMAND(LaunchSwitchboard, "Launch Switchboard", "Launch Switchboard", EUserInterfaceActionType::Button, FInputChord());
+		UI_COMMAND(LaunchSwitchboardListener, "Launch Switchboard Listener", "Launch Switchboard Listener", EUserInterfaceActionType::Button, FInputChord());
 	}
 
-	TSharedPtr<FUICommandInfo> TriggerToolbarButtonCmd;
+	TSharedPtr<FUICommandInfo> LaunchSwitchboard;
+	TSharedPtr<FUICommandInfo> LaunchSwitchboardListener;
 };
 
 
@@ -41,21 +44,29 @@ struct FSwitchboardMenuEntryImpl
 	FSwitchboardMenuEntryImpl()
 	{
 		FSwitchboardUICommands::Register();
-		const TSharedPtr<FUICommandList> Actions = MakeShared<FUICommandList>();
 
-		Actions->MapAction(FSwitchboardUICommands::Get().TriggerToolbarButtonCmd,
-			FExecuteAction::CreateRaw(this, &FSwitchboardMenuEntryImpl::OnLaunchSwitchboardClicked),
-			FCanExecuteAction::CreateRaw(this, &FSwitchboardMenuEntryImpl::IsToolbarButtonEnabled));
+		Actions->MapAction(FSwitchboardUICommands::Get().LaunchSwitchboard,
+			FExecuteAction::CreateRaw(this, &FSwitchboardMenuEntryImpl::OnLaunchSwitchboard),
+			FCanExecuteAction());
 
+		Actions->MapAction(FSwitchboardUICommands::Get().LaunchSwitchboardListener,
+			FExecuteAction::CreateRaw(this, &FSwitchboardMenuEntryImpl::OnLaunchSwitchboardListener),
+			FCanExecuteAction());
+
+		AddMenu();
+	}
+
+	void AddMenu()
+	{
 		UToolMenu* Menu = UToolMenus::Get()->ExtendMenu("LevelEditor.LevelEditorToolBar.User");
-		FToolMenuSection& Section = Menu->FindOrAddSection("Switchboard");
+		if (Menu->FindSection("Switchboard"))
+		{
+			return;
+		}
 
-		FToolMenuEntry SwitchboardButtonEntry = FToolMenuEntry::InitToolBarButton(
-			FSwitchboardUICommands::Get().TriggerToolbarButtonCmd,
-			TAttribute<FText>::Create([this]() { return LOCTEXT("LaunchSwitchboard", "Switchboard"); }),
-			TAttribute<FText>::Create([this]() { return LOCTEXT("LaunchSwitchboardTooltip", "Launch Switchboard"); }),
-			TAttribute<FSlateIcon>::Create([this]() { return GetToolbarButtonIcon(); })
-		);
+		FToolMenuSection& Section = Menu->AddSection("Switchboard");
+
+		FToolMenuEntry SwitchboardButtonEntry = FToolMenuEntry::InitToolBarButton(FSwitchboardUICommands::Get().LaunchSwitchboard);
 		SwitchboardButtonEntry.SetCommandList(Actions);
 
 		const FToolMenuEntry SwitchboardComboEntry = FToolMenuEntry::InitComboButton(
@@ -72,6 +83,15 @@ struct FSwitchboardMenuEntryImpl
 		Section.AddEntry(SwitchboardComboEntry);
 	}
 
+	void RemoveMenu()
+	{
+		UToolMenu* Menu = UToolMenus::Get()->ExtendMenu("LevelEditor.LevelEditorToolBar.User");
+		if (Menu->FindSection("Switchboard"))
+		{
+			Menu->RemoveSection("Switchboard");
+		}
+	}
+
 	~FSwitchboardMenuEntryImpl()
 	{
 		if (!IsEngineExitRequested() && UObjectInitialized())
@@ -87,12 +107,7 @@ struct FSwitchboardMenuEntryImpl
 
 		MenuBuilder.BeginSection("Switchboard", LOCTEXT("SwitchboardListener", "Listener"));
 		{
-			MenuBuilder.AddMenuEntry(
-				LOCTEXT("ListenerLaunchLabel", "Launch SwitchboardListener"),
-				LOCTEXT("ListenerLaunchTooltip", "Launches the SwitchboardListener with the settings from Editor Settings."),
-				FSlateIcon(),
-				FUIAction(FExecuteAction::CreateRaw(this, &FSwitchboardMenuEntryImpl::OnLaunchListenerClicked))
-			);
+			MenuBuilder.AddMenuEntry(FSwitchboardUICommands::Get().LaunchSwitchboardListener);
 
 #if SB_LISTENER_AUTOLAUNCH
 			MenuBuilder.AddMenuSeparator();
@@ -119,7 +134,7 @@ struct FSwitchboardMenuEntryImpl
 		return MenuBuilder.MakeWidget();
 	}
 
-	void OnLaunchListenerClicked()
+	void OnLaunchSwitchboardListener()
 	{
 		const FString ListenerPath = GetDefault<USwitchboardEditorSettings>()->GetListenerPlatformPath();
 		if (!FPaths::FileExists(ListenerPath))
@@ -161,8 +176,15 @@ struct FSwitchboardMenuEntryImpl
 	}
 #endif // #if SB_LISTENER_AUTOLAUNCH
 
-	void OnLaunchSwitchboardClicked()
+	void OnLaunchSwitchboard()
 	{
+		const FSwitchboardVerifyResult& VerifyResult = FSwitchboardEditorModule::Get().GetVerifyResult().Get();
+		if (VerifyResult.Summary != FSwitchboardVerifyResult::ESummary::Success)
+		{
+			SSwitchboardSetupWizard::OpenWindow();
+			return;
+		}
+
 		if (FSwitchboardEditorModule::Get().LaunchSwitchboard())
 		{
 			UE_LOG(LogSwitchboardPlugin, Display, TEXT("Successfully started Switchboard"));
@@ -174,18 +196,10 @@ struct FSwitchboardMenuEntryImpl
 		}
 	}
 
-	FSlateIcon GetToolbarButtonIcon() const
-	{
-		return FSlateIcon(FSwitchboardEditorStyle::Get().GetStyleSetName(), FSwitchboardEditorStyle::NAME_SwitchboardBrush, FSwitchboardEditorStyle::NAME_SwitchboardBrush);
-	}
-
-	bool IsToolbarButtonEnabled()
-	{
-		return true;
-	}
-
 public:
 	static TUniquePtr<FSwitchboardMenuEntryImpl> Implementation;
+
+	TSharedPtr<FUICommandList> Actions = MakeShared<FUICommandList>();
 };
 
 TUniquePtr<FSwitchboardMenuEntryImpl> FSwitchboardMenuEntryImpl::Implementation;
@@ -198,6 +212,22 @@ void FSwitchboardMenuEntry::Register()
 	}
 }
 
+void FSwitchboardMenuEntry::AddMenu()
+{
+	if (FSwitchboardMenuEntryImpl::Implementation)
+	{
+		FSwitchboardMenuEntryImpl::Implementation->AddMenu();
+	}
+}
+
+void FSwitchboardMenuEntry::RemoveMenu()
+{
+	if (FSwitchboardMenuEntryImpl::Implementation)
+	{
+		FSwitchboardMenuEntryImpl::Implementation->RemoveMenu();
+	}
+}
+
 void FSwitchboardMenuEntry::Unregister()
 {
 	FSwitchboardMenuEntryImpl::Implementation.Reset();
diff --git a/Engine/Plugins/VirtualProduction/Switchboard/Source/SwitchboardEditor/Private/SwitchboardMenuEntry.h b/Engine/Plugins/VirtualProduction/Switchboard/Source/SwitchboardEditor/Private/SwitchboardMenuEntry.h
--- a/Engine/Plugins/VirtualProduction/Switchboard/Source/SwitchboardEditor/Private/SwitchboardMenuEntry.h
+++ b/Engine/Plugins/VirtualProduction/Switchboard/Source/SwitchboardEditor/Private/SwitchboardMenuEntry.h
@@ -9,5 +9,7 @@ class FSwitchboardMenuEntry
 {
 public:
 	static void Register();
+	static void AddMenu();
+	static void RemoveMenu();
 	static void Unregister();
 };
diff --git a/Engine/Plugins/VirtualProduction/Switchboard/Source/SwitchboardEditor/Private/SwitchboardScriptInterop.cpp b/Engine/Plugins/VirtualProduction/Switchboard/Source/SwitchboardEditor/Private/SwitchboardScriptInterop.cpp
--- a/Engine/Plugins/VirtualProduction/Switchboard/Source/SwitchboardEditor/Private/SwitchboardScriptInterop.cpp
+++ b/Engine/Plugins/VirtualProduction/Switchboard/Source/SwitchboardEditor/Private/SwitchboardScriptInterop.cpp
@@ -146,7 +146,7 @@ TFuture<FSwitchboardVerifyResult> FSwitchboardVerifyResult::RunVerify(FStringVie
 
 			if (!FPaths::FileExists(Verify.PythonExe))
 			{
-				UE_LOG(LogSwitchboardPlugin, Warning, TEXT("RunVerify: Python interpreter not found: %s"), *Verify.PythonExe);
+				UE_LOG(LogSwitchboardPlugin, Log, TEXT("RunVerify: Python interpreter not found: %s"), *Verify.PythonExe);
 				Result.Summary = FSwitchboardVerifyResult::ESummary::InterpreterMissing;
 				return Result;
 			}
@@ -155,7 +155,7 @@ TFuture<FSwitchboardVerifyResult> FSwitchboardVerifyResult::RunVerify(FStringVie
 			const FString Args = FString::Printf(TEXT("\"%s\" verify --output-json"), *SbSetupPath);
 			if (!Verify.Run(Args))
 			{
-				UE_LOG(LogSwitchboardPlugin, Warning, TEXT("RunVerify: Failed to launch: \"%s\" %s"),
+				UE_LOG(LogSwitchboardPlugin, Log, TEXT("RunVerify: Failed to launch: \"%s\" %s"),
 					*Verify.PythonExe, *Args);
 				Result.Summary = FSwitchboardVerifyResult::ESummary::LaunchFailed;
 				return Result;
diff --git a/Engine/Plugins/VirtualProduction/Switchboard/Source/SwitchboardEditor/Private/SwitchboardSettingsCustomization.cpp b/Engine/Plugins/VirtualProduction/Switchboard/Source/SwitchboardEditor/Private/SwitchboardSettingsCustomization.cpp
--- a/Engine/Plugins/VirtualProduction/Switchboard/Source/SwitchboardEditor/Private/SwitchboardSettingsCustomization.cpp
+++ b/Engine/Plugins/VirtualProduction/Switchboard/Source/SwitchboardEditor/Private/SwitchboardSettingsCustomization.cpp
@@ -27,49 +27,15 @@ TSharedRef<IDetailCustomization> FSwitchboardEditorSettingsCustomization::MakeIn
 }
 
 
-FSwitchboardEditorSettingsCustomization::ESwitchboardInstallState FSwitchboardEditorSettingsCustomization::GetInstallState() const
-{
-	FSwitchboardEditorModule& SbEditorModule = FSwitchboardEditorModule::Get();
-	TSharedFuture<FSwitchboardVerifyResult> VerifyResult = SbEditorModule.GetVerifyResult();
-
-	if (!VerifyResult.IsReady())
-	{
-		return ESwitchboardInstallState::VerifyInProgress;
-	}
-
-	if (VerifyResult.Get().Summary != FSwitchboardVerifyResult::ESummary::Success)
-	{
-		return ESwitchboardInstallState::NeedInstallOrRepair;
-	}
-
-#if SWITCHBOARD_SHORTCUTS
-	using namespace UE::Switchboard::Private::Shorcuts;
-
-	const bool bListenerDesktopShortcutExists = SbEditorModule.DoesShortcutExist(EShortcutApp::Listener, EShortcutLocation::Desktop) == EShortcutCompare::AlreadyExists;
-	const bool bListenerProgramsShortcutExists = SbEditorModule.DoesShortcutExist(EShortcutApp::Listener, EShortcutLocation::Programs) == EShortcutCompare::AlreadyExists;
-	const bool bAppDesktopShortcutExists = SbEditorModule.DoesShortcutExist(EShortcutApp::Switchboard, EShortcutLocation::Desktop) == EShortcutCompare::AlreadyExists;
-	const bool bAppProgramsShortcutExists = SbEditorModule.DoesShortcutExist(EShortcutApp::Switchboard, EShortcutLocation::Programs) == EShortcutCompare::AlreadyExists;
-
-	if (!bListenerDesktopShortcutExists || !bListenerProgramsShortcutExists
-		|| !bAppDesktopShortcutExists || !bAppProgramsShortcutExists)
-	{
-		return ESwitchboardInstallState::ShortcutsMissing;
-	}
-#endif // #if SWITCHBOARD_SHORTCUTS
-
-	return ESwitchboardInstallState::Nominal;
-}
-
-
 FText FSwitchboardEditorSettingsCustomization::GetHealthRowHintText() const
 {
-	switch (GetInstallState())
+	switch (FSwitchboardEditorModule::Get().GetSwitchboardInstallState())
 	{
-		case ESwitchboardInstallState::Nominal:
+		case FSwitchboardEditorModule::ESwitchboardInstallState::Nominal:
 			return LOCTEXT("ReinstallEnvHint", "Switchboard requires additional library/framework dependencies to launch, which are installed. If you encounter issues launching Switchboard, it may help to reinstall them.");
-		case ESwitchboardInstallState::ShortcutsMissing:
+		case FSwitchboardEditorModule::ESwitchboardInstallState::ShortcutsMissing:
 			return LOCTEXT("AddShortcutsHint", "Add shortcuts to launch Switchboard from the Desktop and Start Menu.");
-		case ESwitchboardInstallState::NeedInstallOrRepair:
+		case FSwitchboardEditorModule::ESwitchboardInstallState::NeedInstallOrRepair:
 		default:
 			return LOCTEXT("InstallEnvHint", "Switchboard requires additional library/framework dependencies to launch, please install them.");
 	}
@@ -78,13 +44,13 @@ FText FSwitchboardEditorSettingsCustomization::GetHealthRowHintText() const
 
 FText FSwitchboardEditorSettingsCustomization::GetHealthRowButtonText() const
 {
-	switch (GetInstallState())
+	switch (FSwitchboardEditorModule::Get().GetSwitchboardInstallState())
 	{
-		case ESwitchboardInstallState::Nominal:
+		case FSwitchboardEditorModule::ESwitchboardInstallState::Nominal:
 			return LOCTEXT("ReinstallEnvButton", "Reinstall Dependencies");
-		case ESwitchboardInstallState::ShortcutsMissing:
+		case FSwitchboardEditorModule::ESwitchboardInstallState::ShortcutsMissing:
 			return LOCTEXT("AddShortcutsButton", "Add Shortcuts");
-		case ESwitchboardInstallState::NeedInstallOrRepair:
+		case FSwitchboardEditorModule::ESwitchboardInstallState::NeedInstallOrRepair:
 		default:
 			return LOCTEXT("InstallEnvButton", "Install Dependencies");
 	}
@@ -93,12 +59,12 @@ FText FSwitchboardEditorSettingsCustomization::GetHealthRowButtonText() const
 
 const FSlateBrush* FSwitchboardEditorSettingsCustomization::GetHealthRowBorderBrush() const
 {
-	switch (GetInstallState())
+	switch (FSwitchboardEditorModule::Get().GetSwitchboardInstallState())
 	{
-		case ESwitchboardInstallState::Nominal:
+		case FSwitchboardEditorModule::ESwitchboardInstallState::Nominal:
 			return FSwitchboardEditorStyle::Get().GetBrush("Settings.RowBorder.Nominal");
-		case ESwitchboardInstallState::ShortcutsMissing:
-		case ESwitchboardInstallState::NeedInstallOrRepair:
+		case FSwitchboardEditorModule::ESwitchboardInstallState::ShortcutsMissing:
+		case FSwitchboardEditorModule::ESwitchboardInstallState::NeedInstallOrRepair:
 		default:
 			return FSwitchboardEditorStyle::Get().GetBrush("Settings.RowBorder.Warning");
 	}
@@ -107,12 +73,12 @@ const FSlateBrush* FSwitchboardEditorSettingsCustomization::GetHealthRowBorderBr
 
 const FSlateBrush* FSwitchboardEditorSettingsCustomization::GetHealthRowIconBrush() const
 {
-	switch (GetInstallState())
+	switch (FSwitchboardEditorModule::Get().GetSwitchboardInstallState())
 	{
-		case ESwitchboardInstallState::Nominal:
+		case FSwitchboardEditorModule::ESwitchboardInstallState::Nominal:
 			return FSwitchboardEditorStyle::Get().GetBrush("Settings.Icons.Nominal");
-		case ESwitchboardInstallState::ShortcutsMissing:
-		case ESwitchboardInstallState::NeedInstallOrRepair:
+		case FSwitchboardEditorModule::ESwitchboardInstallState::ShortcutsMissing:
+		case FSwitchboardEditorModule::ESwitchboardInstallState::NeedInstallOrRepair:
 		default:
 			return FSwitchboardEditorStyle::Get().GetBrush("Settings.Icons.Warning");
 	}
@@ -121,6 +87,10 @@ const FSlateBrush* FSwitchboardEditorSettingsCustomization::GetHealthRowIconBrus
 
 void FSwitchboardEditorSettingsCustomization::CustomizeDetails(IDetailLayoutBuilder& DetailLayout)
 {
+	// Force the status to refresh whenever the editor settings are closed and reopened.
+	const bool bForceRefresh = true;
+	FSwitchboardEditorModule::Get().GetVerifyResult(bForceRefresh);
+
 	IDetailCategoryBuilder& SwitchboardCategory = DetailLayout.EditCategory("Switchboard");
 
 	FText HealthFilterString = FText::GetEmpty();
@@ -162,14 +132,11 @@ void FSwitchboardEditorSettingsCustomization::CustomizeDetails(IDetailLayoutBuil
 				[
 					SNew(SButton)
 					.Text(TAttribute<FText>::CreateSP(this, &FSwitchboardEditorSettingsCustomization::GetHealthRowButtonText))
-					.OnClicked_Lambda([ThisWeak = TWeakPtr<FSwitchboardEditorSettingsCustomization>(SharedThis(this))]() {
+					.OnClicked_Lambda([this]() {
 						SSwitchboardSetupWizard::EWizardPage StartPage = SSwitchboardSetupWizard::EWizardPage::Intro;
-						if (TSharedPtr<FSwitchboardEditorSettingsCustomization> This = ThisWeak.Pin())
+						if (FSwitchboardEditorModule::Get().GetSwitchboardInstallState() == FSwitchboardEditorModule::ESwitchboardInstallState::ShortcutsMissing)
 						{
-							if (This->GetInstallState() == ESwitchboardInstallState::ShortcutsMissing)
-							{
-								StartPage = SSwitchboardSetupWizard::EWizardPage::Shortcuts;
-							}
+							StartPage = SSwitchboardSetupWizard::EWizardPage::Shortcuts;
 						}
 						SSwitchboardSetupWizard::OpenWindow(StartPage);
 						return FReply::Handled();
diff --git a/Engine/Plugins/VirtualProduction/Switchboard/Source/SwitchboardEditor/Private/SwitchboardSettingsCustomization.h b/Engine/Plugins/VirtualProduction/Switchboard/Source/SwitchboardEditor/Private/SwitchboardSettingsCustomization.h
--- a/Engine/Plugins/VirtualProduction/Switchboard/Source/SwitchboardEditor/Private/SwitchboardSettingsCustomization.h
+++ b/Engine/Plugins/VirtualProduction/Switchboard/Source/SwitchboardEditor/Private/SwitchboardSettingsCustomization.h
@@ -14,17 +14,6 @@ class FSwitchboardEditorSettingsCustomization : public IDetailCustomization
 public:
 	static TSharedRef<IDetailCustomization> MakeInstance();
 
-	enum class ESwitchboardInstallState
-	{
-		VerifyInProgress,
-
-		NeedInstallOrRepair,
-		ShortcutsMissing,
-		Nominal,
-	};
-
-	ESwitchboardInstallState GetInstallState() const;
-
 	FText GetHealthRowHintText() const;
 	FText GetHealthRowButtonText() const;
 	const FSlateBrush* GetHealthRowBorderBrush() const;
diff --git a/Engine/Plugins/VirtualProduction/Switchboard/Source/SwitchboardEditor/Private/SwitchboardSetupWizard.cpp b/Engine/Plugins/VirtualProduction/Switchboard/Source/SwitchboardEditor/Private/SwitchboardSetupWizard.cpp
--- a/Engine/Plugins/VirtualProduction/Switchboard/Source/SwitchboardEditor/Private/SwitchboardSetupWizard.cpp
+++ b/Engine/Plugins/VirtualProduction/Switchboard/Source/SwitchboardEditor/Private/SwitchboardSetupWizard.cpp
@@ -29,16 +29,17 @@ void SSwitchboardSetupWizard::OpenWindow(EWizardPage StartPage /* = EWizardPage:
 		.SupportsMaximize(false)
 		.ClientSize(FVector2D(700, 350));
 
-	const TSharedRef<SSwitchboardSetupWizard> Wizard = SNew(SSwitchboardSetupWizard, ModalWindow);
-	Wizard->SwitchToPage(StartPage);
+	const TSharedRef<SSwitchboardSetupWizard> Wizard = SNew(SSwitchboardSetupWizard, ModalWindow, StartPage);
 	ModalWindow->SetContent(Wizard);
 
 	GEditor->EditorAddModalWindow(ModalWindow);
 }
 
 
-void SSwitchboardSetupWizard::Construct(const FArguments& InArgs, const TSharedRef<SWindow>& InWindow)
+void SSwitchboardSetupWizard::Construct(const FArguments& InArgs, const TSharedRef<SWindow>& InWindow, EWizardPage InStartPage /* = EWizardPage::Intro */)
 {
+	StartPage = InStartPage;
+
 	ModalWindow = InWindow;
 	ModalWindow->SetRequestDestroyWindowOverride(FRequestDestroyWindowOverride::CreateSP(SharedThis(this), &SSwitchboardSetupWizard::Handle_RequestDestroyWindow));
 
@@ -69,6 +70,8 @@ void SSwitchboardSetupWizard::Construct(const FArguments& InArgs, const TSharedR
 		Construct_Page_Complete()
 	];
 
+	SwitchToPage(StartPage);
+
 	ChildSlot
 	[
 		SNew(SOverlay)
@@ -141,11 +144,8 @@ TSharedRef<SWidget> SSwitchboardSetupWizard::Construct_Page_Intro()
 			[
 				SNew(SButton)
 				.Text(LOCTEXT("Wizard_Intro_CancelButton", "Cancel"))
-				.OnClicked_Lambda([ThisWeak = TWeakPtr<SSwitchboardSetupWizard>(SharedThis(this))]() {
-					if (TSharedPtr<SSwitchboardSetupWizard> This = ThisWeak.Pin())
-					{
-						This->ModalWindow->RequestDestroyWindow();
-					}
+				.OnClicked_Lambda([this]() {
+					ModalWindow->RequestDestroyWindow();
 					return FReply::Handled();
 				})
 			]
@@ -162,14 +162,8 @@ TSharedRef<SWidget> SSwitchboardSetupWizard::Construct_Page_InstallProgress()
 		.VAlign(VAlign_Center)
 		[
 			SNew(STextBlock)
-			.Text_Lambda([ThisWeak = TWeakPtr<SSwitchboardSetupWizard>(SharedThis(this))]() {
-				EInstallProgress Progress = EInstallProgress::Idle;
-				if (TSharedPtr<SSwitchboardSetupWizard> This = ThisWeak.Pin())
-				{
-					Progress = This->GetProgress();
-				}
-
-				switch (Progress)
+			.Text_Lambda([this]() {
+				switch (GetProgress())
 				{
 					case EInstallProgress::Running: return LOCTEXT("Wizard_Progress_InstallRunning", "Installation in progress...");
 					case EInstallProgress::Succeeded: return LOCTEXT("Wizard_Progress_InstallSucceeded", "Installation completed successfully!");
@@ -198,22 +192,8 @@ TSharedRef<SWidget> SSwitchboardSetupWizard::Construct_Page_InstallProgress()
 			.AutoWidth()
 			[
 				SNew(SPrimaryButton)
-				.Text_Lambda([ThisWeak = TWeakPtr<SSwitchboardSetupWizard>(SharedThis(this))]() {
-					EInstallProgress Progress = EInstallProgress::Idle;
-					if (TSharedPtr<SSwitchboardSetupWizard> This = ThisWeak.Pin())
-					{
-						Progress = This->GetProgress();
-					}
-					return Progress != EInstallProgress::Failed ? LOCTEXT("Wizard_Progress_ContinueButton", "Continue") : LOCTEXT("Wizard_Progress_ExitButton", "Exit");
-				})
-				.IsEnabled_Lambda([ThisWeak = TWeakPtr<SSwitchboardSetupWizard>(SharedThis(this))]() {
-					EInstallProgress Progress = EInstallProgress::Idle;
-					if (TSharedPtr<SSwitchboardSetupWizard> This = ThisWeak.Pin())
-					{
-						Progress = This->GetProgress();
-					}
-					return Progress != EInstallProgress::Running;
-				})
+				.Text_Lambda([this]() { return GetProgress() != EInstallProgress::Failed ? LOCTEXT("Wizard_Progress_ContinueButton", "Continue") : LOCTEXT("Wizard_Progress_ExitButton", "Exit"); })
+				.IsEnabled_Lambda([this]() { return GetProgress() != EInstallProgress::Running; })
 				.OnClicked(this, &SSwitchboardSetupWizard::Handle_InstallProgress_ContinueClicked)
 			]
 			+ SHorizontalBox::Slot()
@@ -222,12 +202,9 @@ TSharedRef<SWidget> SSwitchboardSetupWizard::Construct_Page_InstallProgress()
 			[
 				SNew(SButton)
 				.Text(LOCTEXT("Wizard_Progress_CancelButton", "Cancel"))
-				.OnClicked_Lambda([ThisWeak = TWeakPtr<SSwitchboardSetupWizard>(SharedThis(this))]() {
-					if (TSharedPtr<SSwitchboardSetupWizard> This = ThisWeak.Pin())
-					{
-						// Handle_RequestDestroyWindow will prompt the user to confirm if interrupting installer
-						This->ModalWindow->RequestDestroyWindow();
-					}
+				.OnClicked_Lambda([this]() {
+					// Handle_RequestDestroyWindow will prompt the user to confirm if interrupting installer
+					ModalWindow->RequestDestroyWindow();
 					return FReply::Handled();
 				})
 			]
@@ -258,7 +235,9 @@ TSharedRef<SWidget> SSwitchboardSetupWizard::Construct_Page_Shortcuts()
 			.Padding(24.0f, 0.0f, 12.0f, 0.0f)
 			[
 				SAssignNew(DesktopShortcutCheckbox, SCheckBox)
-#if !SWITCHBOARD_SHORTCUTS
+#if SWITCHBOARD_SHORTCUTS
+				.IsChecked(ECheckBoxState::Checked)
+#else
 				.IsEnabled(false)
 #endif
 			]
@@ -284,7 +263,9 @@ TSharedRef<SWidget> SSwitchboardSetupWizard::Construct_Page_Shortcuts()
 			.Padding(24.0f, 0.0f, 12.0f, 0.0f)
 			[
 				SAssignNew(ProgramsShortcutCheckbox, SCheckBox)
-#if !SWITCHBOARD_SHORTCUTS
+#if SWITCHBOARD_SHORTCUTS
+				.IsChecked(ECheckBoxState::Checked)
+#else
 				.Visibility(EVisibility::Collapsed)
 #endif
 			]
@@ -313,24 +294,29 @@ TSharedRef<SWidget> SSwitchboardSetupWizard::Construct_Page_Shortcuts()
 			[
 				SNew(SPrimaryButton)
 				.Text(LOCTEXT("Wizard_Shortcuts_CreateButton", "Create"))
-				.OnClicked_Lambda([ThisWeak = TWeakPtr<SSwitchboardSetupWizard>(SharedThis(this))]() {
-					if (TSharedPtr<SSwitchboardSetupWizard> This = ThisWeak.Pin())
-					{
+				.OnClicked_Lambda([this]() {
 #if SWITCHBOARD_SHORTCUTS
-						if (This->DesktopShortcutCheckbox->IsChecked())
-						{
-							FSwitchboardEditorModule::Get().CreateOrUpdateShortcut(FSwitchboardEditorModule::EShortcutApp::Switchboard, FSwitchboardEditorModule::EShortcutLocation::Desktop);
-							FSwitchboardEditorModule::Get().CreateOrUpdateShortcut(FSwitchboardEditorModule::EShortcutApp::Listener, FSwitchboardEditorModule::EShortcutLocation::Desktop);
-						}
-
-						if (This->ProgramsShortcutCheckbox->IsChecked())
-						{
-							FSwitchboardEditorModule::Get().CreateOrUpdateShortcut(FSwitchboardEditorModule::EShortcutApp::Switchboard, FSwitchboardEditorModule::EShortcutLocation::Programs);
-							FSwitchboardEditorModule::Get().CreateOrUpdateShortcut(FSwitchboardEditorModule::EShortcutApp::Listener, FSwitchboardEditorModule::EShortcutLocation::Programs);
-						}
+					if (DesktopShortcutCheckbox->IsChecked())
+					{
+						FSwitchboardEditorModule::Get().CreateOrUpdateShortcut(FSwitchboardEditorModule::EShortcutApp::Switchboard, FSwitchboardEditorModule::EShortcutLocation::Desktop);
+						FSwitchboardEditorModule::Get().CreateOrUpdateShortcut(FSwitchboardEditorModule::EShortcutApp::Listener, FSwitchboardEditorModule::EShortcutLocation::Desktop);
+					}
+
+					if (ProgramsShortcutCheckbox->IsChecked())
+					{
+						FSwitchboardEditorModule::Get().CreateOrUpdateShortcut(FSwitchboardEditorModule::EShortcutApp::Switchboard, FSwitchboardEditorModule::EShortcutLocation::Programs);
+						FSwitchboardEditorModule::Get().CreateOrUpdateShortcut(FSwitchboardEditorModule::EShortcutApp::Listener, FSwitchboardEditorModule::EShortcutLocation::Programs);
+					}
 #endif
-						This->SwitchToPage(EWizardPage::Autolaunch);
+
+					// If the user just wanted to create shortcuts, we end here.
+					if (StartPage == EWizardPage::Shortcuts)
+					{
+						ModalWindow->RequestDestroyWindow();
+						return FReply::Handled();
 					}
+
+					SwitchToPage(EWizardPage::Autolaunch);
 					return FReply::Handled();
 				})
 			]
@@ -339,12 +325,16 @@ TSharedRef<SWidget> SSwitchboardSetupWizard::Construct_Page_Shortcuts()
 			.Padding(16.0f, 0.0f, 0.0f, 0.0f)
 			[
 				SNew(SButton)
-				.Text(LOCTEXT("Wizard_Shortcuts_SkipButton", "Cancel"))
-				.OnClicked_Lambda([ThisWeak = TWeakPtr<SSwitchboardSetupWizard>(SharedThis(this))]() {
-					if (TSharedPtr<SSwitchboardSetupWizard> This = ThisWeak.Pin())
+				.Text(LOCTEXT("Wizard_Shortcuts_SkipButton", "Skip"))
+				.OnClicked_Lambda([this]() {
+					// If the user just wanted to create shortcuts, we end here.
+					if (StartPage == EWizardPage::Shortcuts)
 					{
-						This->SwitchToPage(EWizardPage::Autolaunch);
+						ModalWindow->RequestDestroyWindow();
+						return FReply::Handled();
 					}
+
+					SwitchToPage(EWizardPage::Autolaunch);
 					return FReply::Handled();
 				})
 			]
@@ -375,7 +365,9 @@ TSharedRef<SWidget> SSwitchboardSetupWizard::Construct_Page_Autolaunch()
 			.Padding(24.0f, 0.0f, 12.0f, 0.0f)
 			[
 				SAssignNew(AutolaunchCheckbox, SCheckBox)
-#if !SB_LISTENER_AUTOLAUNCH
+#if SB_LISTENER_AUTOLAUNCH
+				.IsChecked(FSwitchboardEditorModule::Get().IsListenerAutolaunchEnabled() ? ECheckBoxState::Checked : ECheckBoxState::Unchecked)
+#else
 				.IsEnabled(false)
 #endif
 			]
@@ -406,14 +398,11 @@ TSharedRef<SWidget> SSwitchboardSetupWizard::Construct_Page_Autolaunch()
 			[
 				SNew(SPrimaryButton)
 				.Text(LOCTEXT("Wizard_Autolaunch_ContinueButton", "Continue"))
-				.OnClicked_Lambda([ThisWeak = TWeakPtr<SSwitchboardSetupWizard>(SharedThis(this))]() {
-					if (TSharedPtr<SSwitchboardSetupWizard> This = ThisWeak.Pin())
-					{
+				.OnClicked_Lambda([this]() {
 #if SB_LISTENER_AUTOLAUNCH
-						FSwitchboardEditorModule::Get().SetListenerAutolaunchEnabled(This->AutolaunchCheckbox->IsChecked());
+					FSwitchboardEditorModule::Get().SetListenerAutolaunchEnabled(AutolaunchCheckbox->IsChecked());
 #endif
-						This->SwitchToPage(EWizardPage::Complete);
-					}
+					SwitchToPage(EWizardPage::Complete);
 					return FReply::Handled();
 				})
 			]
@@ -423,11 +412,8 @@ TSharedRef<SWidget> SSwitchboardSetupWizard::Construct_Page_Autolaunch()
 			[
 				SNew(SButton)
 				.Text(LOCTEXT("Wizard_Autolaunch_SkipButton", "Skip"))
-				.OnClicked_Lambda([ThisWeak = TWeakPtr<SSwitchboardSetupWizard>(SharedThis(this))]() {
-					if (TSharedPtr<SSwitchboardSetupWizard> This = ThisWeak.Pin())
-					{
-						This->SwitchToPage(EWizardPage::Complete);
-					}
+				.OnClicked_Lambda([this]() {
+					SwitchToPage(EWizardPage::Complete);
 					return FReply::Handled();
 				})
 			]
@@ -459,12 +445,9 @@ TSharedRef<SWidget> SSwitchboardSetupWizard::Construct_Page_Complete()
 			[
 				SNew(SPrimaryButton)
 				.Text(LOCTEXT("Wizard_Complete_LaunchButton", "Launch Switchboard"))
-				.OnClicked_Lambda([ThisWeak = TWeakPtr<SSwitchboardSetupWizard>(SharedThis(this))]() {
+				.OnClicked_Lambda([this]() {
 					FSwitchboardEditorModule::Get().LaunchSwitchboard();
-					if (TSharedPtr<SSwitchboardSetupWizard> This = ThisWeak.Pin())
-					{
-						This->ModalWindow->RequestDestroyWindow();
-					}
+					ModalWindow->RequestDestroyWindow();
 					return FReply::Handled();
 				})
 			]
@@ -474,11 +457,8 @@ TSharedRef<SWidget> SSwitchboardSetupWizard::Construct_Page_Complete()
 			[
 				SNew(SButton)
 				.Text(LOCTEXT("Wizard_Complete_CloseButton", "Close"))
-				.OnClicked_Lambda([ThisWeak = TWeakPtr<SSwitchboardSetupWizard>(SharedThis(this))]() {
-					if (TSharedPtr<SSwitchboardSetupWizard> This = ThisWeak.Pin())
-					{
-						This->ModalWindow->RequestDestroyWindow();
-					}
+				.OnClicked_Lambda([this]() {
+					ModalWindow->RequestDestroyWindow();
 					return FReply::Handled();
 				})
 			]
diff --git a/Engine/Plugins/VirtualProduction/Switchboard/Source/SwitchboardEditor/Private/SwitchboardSetupWizard.h b/Engine/Plugins/VirtualProduction/Switchboard/Source/SwitchboardEditor/Private/SwitchboardSetupWizard.h
--- a/Engine/Plugins/VirtualProduction/Switchboard/Source/SwitchboardEditor/Private/SwitchboardSetupWizard.h
+++ b/Engine/Plugins/VirtualProduction/Switchboard/Source/SwitchboardEditor/Private/SwitchboardSetupWizard.h
@@ -38,7 +38,7 @@ public:
 
 	static void OpenWindow(EWizardPage StartPage = EWizardPage::Intro);
 
-	void Construct(const FArguments& InArgs, const TSharedRef<SWindow>& InWindow);
+	void Construct(const FArguments& InArgs, const TSharedRef<SWindow>& InWindow, EWizardPage InStartPage = EWizardPage::Intro);
 
 protected:
 	TSharedRef<SWidget> Construct_Page_Intro();
@@ -61,6 +61,7 @@ protected:
 	FReply Handle_InstallProgress_ContinueClicked();
 
 protected:
+	EWizardPage StartPage;
 	TSharedPtr<SWindow> ModalWindow;
 	TSharedPtr<SWidgetSwitcher> PageSwitcher;
 
diff --git a/Engine/Plugins/VirtualProduction/Switchboard/Source/SwitchboardEditor/Private/SwitchboardShortcuts.cpp b/Engine/Plugins/VirtualProduction/Switchboard/Source/SwitchboardEditor/Private/SwitchboardShortcuts.cpp
--- a/Engine/Plugins/VirtualProduction/Switchboard/Source/SwitchboardEditor/Private/SwitchboardShortcuts.cpp
+++ b/Engine/Plugins/VirtualProduction/Switchboard/Source/SwitchboardEditor/Private/SwitchboardShortcuts.cpp
@@ -41,18 +41,14 @@ FShortcutParams BuildShortcutParams(EShortcutApp App, EShortcutLocation Location
 			FString VenvPath = GetDefault<USwitchboardEditorSettings>()->VirtualEnvironmentPath.Path;
 			FPaths::MakePlatformFilename(SbPath);
 			FPaths::MakePlatformFilename(VenvPath);
-			//Params.Target = FString::Printf(TEXT("\"%s\" \"%s\""), *SbPath, *VenvPath);
 			Params.Target = SbPath;
 			Params.Args = FString::Printf(TEXT("\"%s\""), *VenvPath);
 			break;
 		}
 		case EShortcutApp::Listener:
 		{
-			const FString ListenerPath = GetDefault<USwitchboardEditorSettings>()->GetListenerPlatformPath();
-			const FString& ListenerArgs = GetDefault<USwitchboardEditorSettings>()->ListenerCommandlineArguments;
-			//Params.Target = FString::Printf(TEXT("\"%s\" %s"), *ListenerPath, *ListenerArgs);
-			Params.Target = ListenerPath;
-			Params.Args = ListenerArgs;
+			Params.Target = GetDefault<USwitchboardEditorSettings>()->GetListenerPlatformPath();
+			Params.Args = GetDefault<USwitchboardEditorSettings>()->ListenerCommandlineArguments;
 			break;
 		}
 		default:
