Add a Qt connection to the group box trigger so that we can properly
update p4 state.

diff --git a/Engine/Plugins/VirtualProduction/Switchboard/Source/Switchboard/switchboard/settings_dialog.py b/Engine/Plugins/VirtualProduction/Switchboard/Source/Switchboard/switchboard/settings_dialog.py
--- a/Engine/Plugins/VirtualProduction/Switchboard/Source/Switchboard/switchboard/settings_dialog.py
+++ b/Engine/Plugins/VirtualProduction/Switchboard/Source/Switchboard/switchboard/settings_dialog.py
@@ -249,14 +249,19 @@ class SettingsDialog(QtCore.QObject):
 
         # Source control setting UI is fully enabled / disabled depending on P4_ENABLED
         self.ui.source_control_settings_group.setChecked(config.P4_ENABLED.get_value())
+        self.ui.source_control_settings_group.toggled.connect(self._on_group_source_control_changed)
         config.P4_ENABLED.signal_setting_changed.connect(self._on_source_control_setting_changed)
         settings_to_show = {key: setting for (key, setting) in config.source_control_settings.items() if setting != config.P4_ENABLED}
-        
+
         self._create_settings_section(settings_to_show, form_layout)
-        
+
+    def _on_group_source_control_changed(self):
+        CONFIG.P4_ENABLED.update_value(self.ui.source_control_settings_group.isChecked())
+        CONFIG.save();
+
     def _on_source_control_setting_changed(self, old_value: bool, new_value: bool):
         self.ui.source_control_settings_group.setChecked(new_value)
-    
+
     def _create_multi_user_server_settings(self, config: Config):
         self.ui.multi_user_settings = CollapsibleGroupBox()
         self.ui.multi_user_settings.setTitle("Multi User Server")
