Fix linux exe paths to that SB can launch Rsync / Multi-user server from Linux

#prefight

diff --git a/Engine/Plugins/VirtualProduction/Switchboard/Source/Switchboard/switchboard/devices/unreal/plugin_unreal.py b/Engine/Plugins/VirtualProduction/Switchboard/Source/Switchboard/switchboard/devices/unreal/plugin_unreal.py
--- a/Engine/Plugins/VirtualProduction/Switchboard/Source/Switchboard/switchboard/devices/unreal/plugin_unreal.py
+++ b/Engine/Plugins/VirtualProduction/Switchboard/Source/Switchboard/switchboard/devices/unreal/plugin_unreal.py
@@ -11,6 +11,7 @@ import ipaddress
 import os
 import pathlib
 import re
+import sys
 import threading
 from typing import Callable, List, Optional, Set
 import uuid
@@ -1936,9 +1937,12 @@ class DeviceUnreal(Device):
         return remote_project_path / 'Saved' / 'Logs'
 
     def get_rsync_path(self):
-        return pathlib.Path(
-            CONFIG.ENGINE_DIR.get_value(self.name),
-            'Extras', 'ThirdPartyNotUE', 'cwrsync', 'bin', 'rsync.exe')
+        if sys.platform.startswith('win'):
+            return pathlib.Path(
+                CONFIG.ENGINE_DIR.get_value(self.name),
+                'Extras', 'ThirdPartyNotUE', 'cwrsync', 'bin', 'rsync.exe')
+        else:
+            return 'rsync'
 
     def fetch_file(self, remote_path):
         program_name = 'retrieve'
diff --git a/Engine/Plugins/VirtualProduction/Switchboard/Source/Switchboard/switchboard/switchboard_application.py b/Engine/Plugins/VirtualProduction/Switchboard/Source/Switchboard/switchboard/switchboard_application.py
--- a/Engine/Plugins/VirtualProduction/Switchboard/Source/Switchboard/switchboard/switchboard_application.py
+++ b/Engine/Plugins/VirtualProduction/Switchboard/Source/Switchboard/switchboard/switchboard_application.py
@@ -131,7 +131,12 @@ class MultiUserApplication:
                              f'{self.exe_path()}. Has it been built?')
                 return
 
-            cmdline = f'start "Multi User Server" "{self.exe_path()}"'
+            cmdline = ''
+            if sys.platform.startswith('win'):
+                cmdline = f'start "Multi User Server" "{self.exe_path()}"'
+            else:
+                cmdline = f'{self.exe_path()}'
+
             cmdline += f' -CONCERTSERVER="{CONFIG.MUSERVER_SERVER_NAME.get_value()}"'
             cmdline += f' {CONFIG.MUSERVER_COMMAND_LINE_ARGUMENTS.get_value()}'
             cmdline += f' {self.get_mu_server_endpoint_arg()}'
