SwitchboardEditor: Pass correct UI command list to menu builder, fixes "Launch Switchboard Listener." Followup to CL 18957940.

diff --git a/Engine/Plugins/VirtualProduction/Switchboard/Source/SwitchboardEditor/Private/SwitchboardMenuEntry.cpp b/Engine/Plugins/VirtualProduction/Switchboard/Source/SwitchboardEditor/Private/SwitchboardMenuEntry.cpp
--- a/Engine/Plugins/VirtualProduction/Switchboard/Source/SwitchboardEditor/Private/SwitchboardMenuEntry.cpp
+++ b/Engine/Plugins/VirtualProduction/Switchboard/Source/SwitchboardEditor/Private/SwitchboardMenuEntry.cpp
@@ -103,7 +103,7 @@ struct FSwitchboardMenuEntryImpl
 	TSharedRef<SWidget> CreateListenerEntries()
 	{
 		const bool bShouldCloseWindowAfterMenuSelection = true;
-		FMenuBuilder MenuBuilder(bShouldCloseWindowAfterMenuSelection, nullptr);
+		FMenuBuilder MenuBuilder(bShouldCloseWindowAfterMenuSelection, Actions);
 
 		MenuBuilder.BeginSection("Switchboard", LOCTEXT("SwitchboardListener", "Listener"));
 		{
