Switchboard: Fix not handling spaces in path to remote log folder when using rsync.

diff --git a/Engine/Plugins/VirtualProduction/Switchboard/Source/Switchboard/switchboard/devices/unreal/plugin_unreal.py b/Engine/Plugins/VirtualProduction/Switchboard/Source/Switchboard/switchboard/devices/unreal/plugin_unreal.py
--- a/Engine/Plugins/VirtualProduction/Switchboard/Source/Switchboard/switchboard/devices/unreal/plugin_unreal.py
+++ b/Engine/Plugins/VirtualProduction/Switchboard/Source/Switchboard/switchboard/devices/unreal/plugin_unreal.py
@@ -1957,7 +1957,7 @@ class DeviceUnreal(Device):
         dest_module = DeviceUnreal.rsync_server.INCOMING_LOGS_MODULE
         dest_path = f'rsync://{dest_endpoint}/{dest_module}/'
 
-        rsync_args = f'{remote_cygpath} {dest_path}'
+        rsync_args = f'"{remote_cygpath}" "{dest_path}"'
 
         puuid, msg = message_protocol.create_start_process_message(
             prog_path=str(rsync_path),
@@ -2090,7 +2090,7 @@ class DeviceUnreal(Device):
             return
 
         insights = InsightsLauncher()
-        insights.launch(args=[str(tracepath)], allow_duplicate=True)
+        insights.launch(args=[f'"{str(tracepath)}"'], allow_duplicate=True)
 
     def on_copy_last_launch_command(self):
         ''' Copies the last launch command to the clipboard'''
