Switchboard: Force Chaos to be deterministic when launching nDisplay.

diff --git a/Engine/Plugins/VirtualProduction/Switchboard/Source/Switchboard/switchboard/devices/ndisplay/plugin_ndisplay.py b/Engine/Plugins/VirtualProduction/Switchboard/Source/Switchboard/switchboard/devices/ndisplay/plugin_ndisplay.py
--- a/Engine/Plugins/VirtualProduction/Switchboard/Source/Switchboard/switchboard/devices/ndisplay/plugin_ndisplay.py
+++ b/Engine/Plugins/VirtualProduction/Switchboard/Source/Switchboard/switchboard/devices/ndisplay/plugin_ndisplay.py
@@ -871,8 +871,7 @@ class DevicenDisplay(DeviceUnreal):
         ])
 
         # Device profile CVars.
-        dp_cvars = self.csettings['ndisplay_dp_cvars'].get_value(self.name)
-        dp_cvars = [cvar.strip() for cvar in dp_cvars if len(cvar.strip()) and len(cvar.split('=')) == 2]
+        dp_cvars = []
 
         # Insights traces parameters
         if CONFIG.INSIGHTS_TRACE_ENABLE.get_value():
@@ -895,6 +894,14 @@ class DevicenDisplay(DeviceUnreal):
             if 'bookmark' in traces.split(','):
                 dp_cvars.append('nDisplay.sync.diag.VBlankMonitoring=1')
 
+        # Always tell Chaos to be deterministic
+        dp_cvars.append('p.Chaos.Solver.Deterministic=1')
+
+        # Add user set dp cvars, overriding any of the forced ones.
+        user_dp_cvars = self.csettings['ndisplay_dp_cvars'].get_value(self.name)
+        user_dp_cvars = [cvar.strip() for cvar in user_dp_cvars if len(cvar.strip()) and len(cvar.split('=')) == 2]
+        dp_cvars = self.add_or_override_cvars(dp_cvars, user_dp_cvars)
+
         # add accumulated dpcvars to args
         if len(dp_cvars):
             args.append(f"-DPCVars=\"{','.join(dp_cvars)}\"")
diff --git a/Engine/Plugins/VirtualProduction/Switchboard/Source/Switchboard/switchboard/devices/unreal/plugin_unreal.py b/Engine/Plugins/VirtualProduction/Switchboard/Source/Switchboard/switchboard/devices/unreal/plugin_unreal.py
--- a/Engine/Plugins/VirtualProduction/Switchboard/Source/Switchboard/switchboard/devices/unreal/plugin_unreal.py
+++ b/Engine/Plugins/VirtualProduction/Switchboard/Source/Switchboard/switchboard/devices/unreal/plugin_unreal.py
@@ -1316,6 +1316,34 @@ class DeviceUnreal(Device):
     def get_utrace_filepath(self):
         return self.get_remote_log_path() / f'{self.name}_{self.runtime_str}.utrace'
 
+    @staticmethod
+    def add_or_override_cvars(cvars:list[str], new_cvars:list[str]):
+        ''' Adds new cvars, or overrides values if already existing.
+
+        Args:
+            cvars     : Cvars to be overridden
+            new_cvars : Cvars to override or add
+
+        Returns:
+            list[str]: Resulting list of cvars.
+        '''
+
+        cvars_map = OrderedDict() # keep the original order
+
+        for cvar in cvars + new_cvars:
+
+            parts = cvar.strip().split('=')
+
+            if len(parts) != 2:
+                continue
+
+            name = parts[0]
+            value = parts[1]
+
+            cvars_map[name.lower()] = (name,value)
+
+        return [f'{cvar[0]}={cvar[1]}' for cvar in cvars_map.values()]
+
     def generate_unreal_command_line_args(self, map_name):
         command_line_args = f'{self.extra_cmdline_args_setting}'
 
@@ -1347,8 +1375,7 @@ class DeviceUnreal(Device):
             command_line_args += f' -ExecCmds="{exec_cmds_expanded}"'
 
         # DPCVars may need to be appended to, so we don't concatenate them until the end.
-        dp_cvars = DeviceUnreal.csettings["dp_cvars"].get_value(self.name)
-        dp_cvars = [cvar.strip() for cvar in dp_cvars if len(cvar.strip().split('=')) == 2]
+        dp_cvars = []
 
         (supported_roles, unsupported_roles) = self.get_vproles()
 
@@ -1376,6 +1403,11 @@ class DeviceUnreal(Device):
         except ValueError:
             LOGGER.warning(f"Invalid Number of GPUs '{max_gpu_count}'")
 
+        # Add user set dp cvars, overriding any of the forced ones.
+        user_dp_cvars = DeviceUnreal.csettings["dp_cvars"].get_value(self.name)
+        user_dp_cvars = [cvar.strip() for cvar in user_dp_cvars if len(cvar.strip().split('=')) == 2]
+        dp_cvars = self.add_or_override_cvars(dp_cvars, user_dp_cvars)
+
         # add accumulated dpcvars to args
         if len(dp_cvars):
             command_line_args += f" -DPCVars=\"{','.join(dp_cvars)}\""
