Switchboard: Fix muserver and tracing settings not properly loaded from config.

diff --git a/Engine/Plugins/VirtualProduction/Switchboard/Source/Switchboard/switchboard/config.py b/Engine/Plugins/VirtualProduction/Switchboard/Source/Switchboard/switchboard/config.py
--- a/Engine/Plugins/VirtualProduction/Switchboard/Source/Switchboard/switchboard/config.py
+++ b/Engine/Plugins/VirtualProduction/Switchboard/Source/Switchboard/switchboard/config.py
@@ -2109,7 +2109,7 @@ class Config(object):
             "tracing_args": StringSetting(
                 "tracing_args",
                 "Unreal Insights Tracing Args",
-                data.get('tracing_arg', 'log,cpu,gpu,frame,bookmark,concert,messaging')
+                data.get('tracing_args', 'log,cpu,gpu,frame,bookmark,concert,messaging')
             ),
             "tracing_stat_events": BoolSetting(
                 "tracing_stat_events",
diff --git a/Engine/Plugins/VirtualProduction/Switchboard/Source/Switchboard/switchboard/switchboard_dialog.py b/Engine/Plugins/VirtualProduction/Switchboard/Source/Switchboard/switchboard/switchboard_dialog.py
--- a/Engine/Plugins/VirtualProduction/Switchboard/Source/Switchboard/switchboard/switchboard_dialog.py
+++ b/Engine/Plugins/VirtualProduction/Switchboard/Source/Switchboard/switchboard/switchboard_dialog.py
@@ -94,7 +94,6 @@ class DeviceAdditionalSettingsUI(QtCore.QObject):
         if is_checked is not self.enable_insights_menu().isChecked():
             self.enable_insights_menu().setChecked(is_checked)
             self._update_enable_state()
-            self.signal_device_widget_tracing.emit(self)
 
     def set_insight_tracing_args(self, value):
         """
@@ -770,7 +769,9 @@ class SwitchboardDialog(QtCore.QObject):
         self.p4_refresh_engine_cl()
         self.refresh_levels()
         self.update_current_config_text()
-        
+        self.refresh_muserver_autojoin()
+        self.refresh_trace_settings()
+
     def update_current_config_text(self):
         # Can be none when current file is deleted
         if SETTINGS.CONFIG is not None:
