Switchboard: Fix missing muserver and trace settings when a new config is created. Alos show project name in window title.

diff --git a/Engine/Plugins/VirtualProduction/Switchboard/Source/Switchboard/switchboard/switchboard_dialog.py b/Engine/Plugins/VirtualProduction/Switchboard/Source/Switchboard/switchboard/switchboard_dialog.py
--- a/Engine/Plugins/VirtualProduction/Switchboard/Source/Switchboard/switchboard/switchboard_dialog.py
+++ b/Engine/Plugins/VirtualProduction/Switchboard/Source/Switchboard/switchboard/switchboard_dialog.py
@@ -408,6 +408,8 @@ class SwitchboardDialog(QtCore.QObject):
             lambda: self.update_current_ip_address_text()
         )
         self.window.current_ip_value.editingFinished.connect(self._try_change_ip_address)
+
+        self.refresh_window_title()
         
     def _try_change_ip_address(self):
         def is_valid_ip_format(address):
@@ -611,6 +613,7 @@ class SwitchboardDialog(QtCore.QObject):
         CONFIG.INSIGHTS_TRACE_ARGS.signal_setting_changed.connect(lambda: self.refresh_trace_settings())
         CONFIG.INSIGHTS_STAT_EVENTS.signal_setting_changed.connect(lambda: self.refresh_trace_settings())
         CONFIG.MUSERVER_AUTO_JOIN.signal_setting_changed.connect(lambda: self.refresh_muserver_autojoin())
+        CONFIG.PROJECT_NAME.signal_setting_changed.connect(lambda: self.refresh_window_title())
 
     def refresh_trace_settings(self):
         self.window.additional_settings.set_insight_trace_state(CONFIG.INSIGHTS_TRACE_ENABLE.get_value())
@@ -771,6 +774,17 @@ class SwitchboardDialog(QtCore.QObject):
         self.update_current_config_text()
         self.refresh_muserver_autojoin()
         self.refresh_trace_settings()
+        self.refresh_window_title()
+
+    def refresh_window_title(self):
+        ''' Updates the window title based on the project name '''
+
+        project_name = CONFIG.PROJECT_NAME.get_value()
+
+        if project_name:
+            self.window.setWindowTitle(f"Switchboard - {project_name}")
+        else:
+            self.window.setWindowTitle(f"Switchboard")
 
     def update_current_config_text(self):
         # Can be none when current file is deleted
@@ -821,6 +835,9 @@ class SwitchboardDialog(QtCore.QObject):
             self.toggle_p4_controls(CONFIG.P4_ENABLED.get_value())
             self.refresh_levels()
             self.update_current_config_text()
+            self.refresh_muserver_autojoin()
+            self.refresh_trace_settings()
+            self.refresh_window_title()
 
     def menu_delete_config(self):
         """
