Switchboard: Use correct path to engine python executable on non-Windows platforms.

diff --git a/Engine/Plugins/VirtualProduction/Switchboard/Source/Switchboard/switchboard/devices/unreal/plugin_unreal.py b/Engine/Plugins/VirtualProduction/Switchboard/Source/Switchboard/switchboard/devices/unreal/plugin_unreal.py
--- a/Engine/Plugins/VirtualProduction/Switchboard/Source/Switchboard/switchboard/devices/unreal/plugin_unreal.py
+++ b/Engine/Plugins/VirtualProduction/Switchboard/Source/Switchboard/switchboard/devices/unreal/plugin_unreal.py
@@ -1020,9 +1020,15 @@ class DeviceUnreal(Device):
             f"{self.name}: Queuing sync for project {project_name} "
             f"(revisions: engine={engine_cl}, project={project_cl})")
 
-        python_path = os.path.normpath(os.path.join(
-            engine_dir, 'Binaries', 'ThirdParty', 'Python3', 'Win64',
-            'python.exe'))
+        if self.platform_binary_directory == 'Win64':
+            python_path = os.path.normpath(os.path.join(
+                engine_dir, 'Binaries', 'ThirdParty', 'Python3',
+                self.platform_binary_directory, 'python.exe'))
+        else:
+            python_path = os.path.normpath(os.path.join(
+                engine_dir, 'Binaries', 'ThirdParty', 'Python3',
+                self.platform_binary_directory, 'bin', 'python3'))
+
         helper_path = os.path.normpath(os.path.join(
             engine_dir, 'Plugins', 'VirtualProduction', 'Switchboard',
             'Source', 'Switchboard', 'sbl_helper.py'))
