Switchboard: Pass correct target platform during builds, rather than hardcoded Win64.

diff --git a/Engine/Plugins/VirtualProduction/Switchboard/Source/Switchboard/switchboard/devices/unreal/plugin_unreal.py b/Engine/Plugins/VirtualProduction/Switchboard/Source/Switchboard/switchboard/devices/unreal/plugin_unreal.py
--- a/Engine/Plugins/VirtualProduction/Switchboard/Source/Switchboard/switchboard/devices/unreal/plugin_unreal.py
+++ b/Engine/Plugins/VirtualProduction/Switchboard/Source/Switchboard/switchboard/devices/unreal/plugin_unreal.py
@@ -1084,7 +1084,7 @@ class DeviceUnreal(Device):
     def build(self):
         if self.exclude_from_build.get_value():
             return
-        
+
         program_name = 'build_project'
 
         # check if it is already on its way:
@@ -1144,27 +1144,36 @@ class DeviceUnreal(Device):
         puuid_dependency = self._build_project(
             puuid_dependency=puuid_dependency)
 
+    @property
+    def target_platform(self) -> str:
+        ''' Returns the Unreal target platform name for e.g. build actions '''
+        # FIXME?: Not strictly correct, but holds for desktop host platforms
+        return self.platform_binary_directory
+
     def _build_project(self, puuid_dependency: Optional[uuid.UUID] = None):
         ubt_args = (
-            'Win64 Development '
+            f'{self.target_platform} Development '
             f'-project="{CONFIG.UPROJECT_PATH.get_value(self.name)}" '
             '-TargetType=Editor -Progress -NoHotReloadFromIDE')
         return self._queue_build(
             'project', ubt_args=ubt_args, puuid_dependency=puuid_dependency)
 
     def _build_mu_server(self, puuid_dependency: Optional[uuid.UUID] = None):
-        ubt_args = 'UnrealMultiUserServer Win64 Development -Progress'
+        ubt_args = (f'UnrealMultiUserServer {self.target_platform} '
+                    'Development -Progress')
         return self._queue_build(
             'mu_server', ubt_args=ubt_args, puuid_dependency=puuid_dependency)
 
     def _build_listener(self, puuid_dependency: Optional[uuid.UUID] = None):
-        ubt_args = 'SwitchboardListener Win64 Development -Progress'
+        ubt_args = (f'SwitchboardListener {self.target_platform} '
+                    'Development -Progress')
         return self._queue_build(
             'listener', ubt_args=ubt_args, puuid_dependency=puuid_dependency)
 
     def _build_shadercompileworker(
             self, puuid_dependency: Optional[uuid.UUID] = None):
-        ubt_args = 'ShaderCompileWorker Win64 Development -Progress'
+        ubt_args = (f'ShaderCompileWorker {self.target_platform} '
+                    'Development -Progress')
         return self._queue_build(
             'shadercw', ubt_args=ubt_args, puuid_dependency=puuid_dependency)
 
