Switchboard: Use Qt functionality for "explore directory," "copy to clipboard," and "open file with default application." Fixes various issues with non-Windows platforms.

diff --git a/Engine/Plugins/VirtualProduction/Switchboard/Source/Switchboard/switchboard/devices/unreal/plugin_unreal.py b/Engine/Plugins/VirtualProduction/Switchboard/Source/Switchboard/switchboard/devices/unreal/plugin_unreal.py
--- a/Engine/Plugins/VirtualProduction/Switchboard/Source/Switchboard/switchboard/devices/unreal/plugin_unreal.py
+++ b/Engine/Plugins/VirtualProduction/Switchboard/Source/Switchboard/switchboard/devices/unreal/plugin_unreal.py
@@ -17,6 +17,7 @@ from typing import Callable, List, Optional, Set
 import uuid
 
 from PySide2 import QtCore
+from PySide2 import QtGui
 from PySide2 import QtWidgets
 
 from switchboard import config_osc as osc
@@ -24,7 +25,7 @@ from switchboard import message_protocol
 from switchboard import switchboard_application
 from switchboard import switchboard_utils as sb_utils
 from switchboard.config import CONFIG, BoolSetting, DirectoryPathSetting, \
-    IntSetting, MultiOptionSetting, OptionSetting, StringSetting, FileSystemPathSetting, \
+    FilePathSetting, IntSetting, MultiOptionSetting, OptionSetting, StringSetting, \
     SETTINGS, DEFAULT_MAP_TEXT, StringListSetting, migrate_comma_separated_string_to_list
 from switchboard.devices.device_base import Device, DeviceStatus, \
     PluginHeaderWidgets
@@ -505,24 +506,24 @@ class DeviceUnreal(Device):
         self.last_launch_command = StringSetting(
             attr_name="last_launch_command",
             nice_name="Last Launch Command",
-            value=kwargs.get("last_launch_command",''),
+            value=kwargs.get("last_launch_command", ''),
             show_ui=False
         )
 
-        self.last_log_path = FileSystemPathSetting(
+        self.last_log_path = FilePathSetting(
             attr_name="last_log_path",
             nice_name="Last Log Path",
-            value=kwargs.get("last_log_path",''),
+            value=kwargs.get("last_log_path", ''),
             show_ui=False
         )
 
-        self.last_trace_path = FileSystemPathSetting(
+        self.last_trace_path = FilePathSetting(
             attr_name="last_trace_path",
             nice_name="Last Insights Trace Path",
-            value=kwargs.get("last_trace_path",''),
+            value=kwargs.get("last_trace_path", ''),
             show_ui=False
         )
-        
+
         self.exclude_from_build = BoolSetting(
             attr_name="exclude_from_build",
             nice_name="Exclude from build",
@@ -2125,7 +2126,10 @@ class DeviceUnreal(Device):
 
     def on_open_last_log(self):
         ''' Opens the last log in your preferred editor '''
-        sb_utils.openfile_with_default_app(self.last_log_path.get_value())
+        path_str = self.last_log_path.get_value()
+        if path_str and pathlib.Path(path_str).is_file():
+            url = QtCore.QUrl.fromLocalFile(path_str)
+            QtGui.QDesktopServices.openUrl(url)
 
     def on_open_last_trace(self):
         ''' Opens the last unreal insights trace '''
@@ -2141,7 +2145,8 @@ class DeviceUnreal(Device):
 
     def on_copy_last_launch_command(self):
         ''' Copies the last launch command to the clipboard'''
-        sb_utils.copy2clipboard(self.last_launch_command.get_value())
+        QtGui.QGuiApplication.clipboard().setText(
+            self.last_launch_command.get_value())
 
 def parse_unreal_tag_file(file_content):
     tags = []
diff --git a/Engine/Plugins/VirtualProduction/Switchboard/Source/Switchboard/switchboard/switchboard_utils.py b/Engine/Plugins/VirtualProduction/Switchboard/Source/Switchboard/switchboard/switchboard_utils.py
--- a/Engine/Plugins/VirtualProduction/Switchboard/Source/Switchboard/switchboard/switchboard_utils.py
+++ b/Engine/Plugins/VirtualProduction/Switchboard/Source/Switchboard/switchboard/switchboard_utils.py
@@ -1,20 +1,13 @@
 # Copyright Epic Games, Inc. All Rights Reserved.
-import typing
-
-from .switchboard_logging import LOGGER
-
-from enum import Enum
-
-from os import listdir
-from os.path import isfile, join
-from zipfile import ZipFile
-
 import datetime
-import subprocess
+from enum import Enum
 import os
-import threading
-import sys
 import pathlib
+import subprocess
+import sys
+import threading
+
+from .switchboard_logging import LOGGER
 
 
 class PriorityModifier(Enum):
@@ -162,6 +155,7 @@ def remove_prefix(str, prefix):
         return str[len(prefix):]
     return str
 
+
 def expand_endpoint(
         endpoint: str, default_addr: str = '0.0.0.0',
         default_port: int = 0) -> str:
@@ -178,37 +172,3 @@ def expand_endpoint(
         port_str = str(default_port)
 
     return f'{addr_str}:{port_str}'
-
-def explore_path(path:str):
-    ''' Opens the os file browser at the specified folder path '''
-
-    if sys.platform.startswith('win'):
-        subprocess.Popen(f'explorer {path}')
-    elif sys.platform.startswith('linux'):
-        subprocess.Popen(f'xdg-open {path}')
-    elif sys.platform.startswith('darwin'):
-        subprocess.Popen(f'open {path}')
-    else:
-        LOGGER.error(f"explore_path not supported in platform '{sys.platform}'")
-        
-
-def copy2clipboard(text:str):
-    if sys.platform.startswith('win'):
-        subprocess.Popen(f'echo[{text.strip()}|clip', shell=True)
-    else:
-        LOGGER.error(f"copy2clipboard not supported in platform '{sys.platform}'")
-
-def openfile_with_default_app(path:str):
-
-    if not pathlib.Path(path).is_file():
-        LOGGER.error(f"openfile_with_default_app could not find '{path}'")
-        return
-
-    if sys.platform.startswith('win'):
-        os.startfile(path)
-    elif sys.platform.startswith('linux'):
-        subprocess.Popen(f'xdg-open {path}')
-    elif sys.platform.startswith('darwin'):
-        subprocess.Popen(f'open {path}')
-    else:
-        LOGGER.error(f"openfile_with_default_app not supported in platform '{sys.platform}'")
diff --git a/Engine/Plugins/VirtualProduction/Switchboard/Source/Switchboard/switchboard/util/collect_logs.py b/Engine/Plugins/VirtualProduction/Switchboard/Source/Switchboard/switchboard/util/collect_logs.py
--- a/Engine/Plugins/VirtualProduction/Switchboard/Source/Switchboard/switchboard/util/collect_logs.py
+++ b/Engine/Plugins/VirtualProduction/Switchboard/Source/Switchboard/switchboard/util/collect_logs.py
@@ -7,19 +7,20 @@ import typing
 from datetime import datetime
 from zipfile import ZipFile
 
-from PySide2 import QtWidgets
+from PySide2 import QtCore, QtGui, QtWidgets
 
 from switchboard.config import Config
 from switchboard.devices.device_base import DeviceStatus
 from switchboard.devices.unreal.plugin_unreal import DeviceUnreal
-from ..switchboard_logging import DEFAULT_LOG_PATH, LOGGER
-
-from ..switchboard_utils import explore_path
+from ..switchboard_logging import LOGGER
 
 
 def open_logs_folder():
-    explore_path(DeviceUnreal.get_log_download_dir())
-    
+    path = DeviceUnreal.get_log_download_dir()
+    if path and path.is_dir():
+        url = QtCore.QUrl.fromLocalFile(str(path))
+        QtGui.QDesktopServices.openUrl(url)
+
 
 def execute_zip_logs_workflow(config: Config, devices: typing.List[DeviceUnreal]):
     """ Zip up all important log files asking the user about preferences in certain edge cases. """
@@ -40,7 +41,7 @@ def execute_zip_logs_workflow(config: Config, devices: typing.List[DeviceUnreal]
 
             return True if continue_answer == QtWidgets.QMessageBox.No else False
         return False
-    
+
     def ask_log_zip_name():
         modtime = datetime.now()
         default_name = f"Logs-{modtime.strftime('%Y.%m.%d-%H.%M.%S')}"
@@ -59,7 +60,7 @@ def execute_zip_logs_workflow(config: Config, devices: typing.List[DeviceUnreal]
         if not file_path.endswith(".zip"):
             file_path += ".zip"
         return file_path
-    
+
     if should_skip_because_running_devices():
         return
 
@@ -80,7 +81,7 @@ def zip_logs(zip_destination: str, config: Config, devices: typing.List[DeviceUn
             return
         else:
             zip_file.write(abs_file_path, zip_path_name)
-        
+
     def copy_db_file(abs_file_path: str, zip_path_name: str, zip_file: ZipFile):
         copy_path = os.path.join(os.path.dirname(zip_destination), os.path.basename(zip_path_name))
         try:
@@ -94,7 +95,7 @@ def zip_logs(zip_destination: str, config: Config, devices: typing.List[DeviceUn
             LOGGER.error(f"Failed to access database file {abs_file_path}")
         except (OSError, IOError): 
             LOGGER.error(f"Failed to create temporary file {copy_path} while backing up original {abs_file_path}")
-    
+
     files_to_remove = []
     with ZipFile(zip_destination, "w") as zip_file:
         iterate_log_files(
@@ -102,19 +103,19 @@ def zip_logs(zip_destination: str, config: Config, devices: typing.List[DeviceUn
             devices,
             lambda abs_file_path, zip_path_name, zip_file=zip_file: process_log_file(abs_file_path, zip_path_name, zip_file)
         )
-        
+
     for file_to_remove in files_to_remove:
         os.remove(file_to_remove)
-        
-        
+
+
 def iterate_log_files(config: Config, devices: typing.List[DeviceUnreal], consumer: typing.Callable[[str, str], None]):
-    """ 
+    """
         Passes the abolute files paths of all important log files:
         - all latest .log files in Switchboard log folder
         - all latest .utrace files in Switchboard log folder
         - all multiuser session files in Programs/UnrealMultiUserServerIntermediateMultiUser
         - multiuser log file ProgramsUnrealMultiUserServerSavedLogsUnrealMultiUserServer.log
-        
+
         The arguments of the consumer:
             1 The absolute file path of the file
             2 A short unique relative file path (includes file name), useful for zipping
@@ -123,7 +124,7 @@ def iterate_log_files(config: Config, devices: typing.List[DeviceUnreal], consum
         for file_path in files:
             if os.path.isfile(file_path):
                 consumer(file_path, os.path.join("Devices", os.path.basename(file_path))) 
-                
+
     def iterate_mu_session_files():
         for current_dir, _, files in os.walk(config.multiuser_server_session_directory_path()):
             for file in files:
@@ -140,5 +141,5 @@ def iterate_log_files(config: Config, devices: typing.List[DeviceUnreal], consum
 
     iterate_mu_session_files()
     consumer(config.multiuser_server_log_path(), os.path.join("MultiUserServer", "UnrealMultiUserServer.log"))
-    
-    consumer(LOGGER.save_log_file(), "Switchboard.log")
\ No newline at end of file
+
+    consumer(LOGGER.save_log_file(), "Switchboard.log")
