Switchboard: Don't potentially query Perforce for every keystroke in the engine CL combo.

diff --git a/Engine/Plugins/VirtualProduction/Switchboard/Source/Switchboard/switchboard/switchboard_dialog.py b/Engine/Plugins/VirtualProduction/Switchboard/Source/Switchboard/switchboard/switchboard_dialog.py
--- a/Engine/Plugins/VirtualProduction/Switchboard/Source/Switchboard/switchboard/switchboard_dialog.py
+++ b/Engine/Plugins/VirtualProduction/Switchboard/Source/Switchboard/switchboard/switchboard_dialog.py
@@ -318,7 +318,10 @@ class SwitchboardDialog(QtCore.QObject):
         self.window.level_combo_box.currentIndexChanged.connect(self._on_selected_level_changed)
         self.window.refresh_levels_button.clicked.connect(self.refresh_levels_incremental)
         self.window.project_cl_combo_box.currentTextChanged.connect(self._set_project_changelist)
-        self.window.engine_cl_combo_box.currentTextChanged.connect(self._set_engine_changelist)
+        self.window.engine_cl_combo_box.currentIndexChanged.connect(
+            lambda _: self._set_engine_changelist(self.window.engine_cl_combo_box.currentText()))
+        self.window.engine_cl_combo_box.lineEdit().editingFinished.connect(
+            lambda: self._set_engine_changelist(self.window.engine_cl_combo_box.currentText()))
         self.window.logger_level_comboBox.currentTextChanged.connect(self.logger_level_comboBox_currentTextChanged)
         self.window.logger_autoscroll_checkbox.stateChanged.connect(self.logger_autoscroll_stateChanged)
         self.window.logger_wrap_checkbox.stateChanged.connect(self.logger_wrap_stateChanged)
