Switchboard: Fix a handful of user-facing references to "UE4".

diff --git a/Engine/Plugins/VirtualProduction/Switchboard/Source/Switchboard/switchboard/add_config_dialog.py b/Engine/Plugins/VirtualProduction/Switchboard/Source/Switchboard/switchboard/add_config_dialog.py
--- a/Engine/Plugins/VirtualProduction/Switchboard/Source/Switchboard/switchboard/add_config_dialog.py
+++ b/Engine/Plugins/VirtualProduction/Switchboard/Source/Switchboard/switchboard/add_config_dialog.py
@@ -544,7 +544,7 @@ class AddConfigDialog(QtWidgets.QDialog):
         self.update_button_box()
 
     def on_browse_engine_dir(self):
-        self.engine_dir = QtWidgets.QFileDialog.getExistingDirectory(self, "Select UE4 engine directory")
+        self.engine_dir = QtWidgets.QFileDialog.getExistingDirectory(self, "Select UE 'Engine' directory")
         self.engine_dir = os.path.normpath(self.engine_dir)
         self.engine_dir_line_edit.setText(self.engine_dir)
 
diff --git a/Engine/Plugins/VirtualProduction/Switchboard/Source/Switchboard/switchboard/config.py b/Engine/Plugins/VirtualProduction/Switchboard/Source/Switchboard/switchboard/config.py
--- a/Engine/Plugins/VirtualProduction/Switchboard/Source/Switchboard/switchboard/config.py
+++ b/Engine/Plugins/VirtualProduction/Switchboard/Source/Switchboard/switchboard/config.py
@@ -2021,7 +2021,7 @@ class Config(object):
                 "engine_dir",
                 "Engine Directory",
                 data.get('engine_dir', ''),
-                tool_tip="Path to UE4 engine directory"
+                tool_tip="Path to UE 'Engine' directory"
             ),
             "build_engine": BoolSetting(
                 "build_engine",
@@ -2084,13 +2084,13 @@ class Config(object):
                 "p4_sync_path",
                 "Perforce Project Path",
                 data.get("p4_sync_path", ''),
-                placeholder_text="//UE4/Project"
+                placeholder_text="//UE/Project"
             ),
             "p4_engine_path": PerforcePathSetting(
                 "p4_engine_path",
                 "Perforce Engine Path",
                 data.get("p4_engine_path", ''),
-                placeholder_text="//UE4/Project/Engine"
+                placeholder_text="//UE/Project/Engine"
             )
         }
 
diff --git a/Engine/Plugins/VirtualProduction/Switchboard/Source/Switchboard/switchboard/devices/unreal/plugin_unreal.py b/Engine/Plugins/VirtualProduction/Switchboard/Source/Switchboard/switchboard/devices/unreal/plugin_unreal.py
--- a/Engine/Plugins/VirtualProduction/Switchboard/Source/Switchboard/switchboard/devices/unreal/plugin_unreal.py
+++ b/Engine/Plugins/VirtualProduction/Switchboard/Source/Switchboard/switchboard/devices/unreal/plugin_unreal.py
@@ -1482,7 +1482,7 @@ class DeviceUnreal(Device):
 
         self.compute_runtime_str()
         engine_path, args = self.generate_unreal_command_line(map_name)
-        LOGGER.info(f"Launching UE4: {engine_path} {args}")
+        LOGGER.info(f"Launching Unreal: {engine_path} {args}")
 
         priority_modifier_str = self.csettings['priority_modifier'].get_value(
             self.name)
