SwitchboardEditor: Don't hide toolbar button on platforms that don't support creation of shortcuts.

diff --git a/Engine/Plugins/VirtualProduction/Switchboard/Source/SwitchboardEditor/Private/SwitchboardEditorModule.cpp b/Engine/Plugins/VirtualProduction/Switchboard/Source/SwitchboardEditor/Private/SwitchboardEditorModule.cpp
--- a/Engine/Plugins/VirtualProduction/Switchboard/Source/SwitchboardEditor/Private/SwitchboardEditorModule.cpp
+++ b/Engine/Plugins/VirtualProduction/Switchboard/Source/SwitchboardEditor/Private/SwitchboardEditorModule.cpp
@@ -200,6 +200,10 @@ TSharedFuture<FSwitchboardVerifyResult> FSwitchboardEditorModule::GetVerifyResul
 			UE_LOG(LogSwitchboardPlugin, Log, TEXT("Verify summary: %d"), Result.Summary);
 			UE_LOG(LogSwitchboardPlugin, Log, TEXT("Verify log: %s"), *Result.Log);
 
+			// On platforms where we support creating shortcuts, once shortcuts have been created,
+			// we hide our toolbar button to yield space.
+#if SWITCHBOARD_SHORTCUTS
+			// This task gets us from the future back onto the main thread prior to manipulating the UI.
 			Async(EAsyncExecution::TaskGraphMainThread, []() {
 				if (FSwitchboardEditorModule::Get().GetSwitchboardInstallState() == ESwitchboardInstallState::Nominal)
 				{
@@ -210,6 +214,7 @@ TSharedFuture<FSwitchboardVerifyResult> FSwitchboardEditorModule::GetVerifyResul
 					FSwitchboardMenuEntry::AddMenu();
 				}
 			});
+#endif
 
 			return Result;
 		});
