Switchboard: sbl_helper: Fix path to GenerateProjectFiles on Linux.

diff --git a/Engine/Plugins/VirtualProduction/Switchboard/Source/Switchboard/sbl_helper.py b/Engine/Plugins/VirtualProduction/Switchboard/Source/Switchboard/sbl_helper.py
--- a/Engine/Plugins/VirtualProduction/Switchboard/Source/Switchboard/sbl_helper.py
+++ b/Engine/Plugins/VirtualProduction/Switchboard/Source/Switchboard/sbl_helper.py
@@ -1014,7 +1014,7 @@ class SbListenerHelper:
                                       'GenerateProjectFiles.bat')
         else:
             gpf_script = os.path.join(self.engine_dir, 'Build', 'BatchFiles',
-                                      'Linux' 'GenerateProjectFiles.sh')
+                                      'Linux', 'GenerateProjectFiles.sh')
 
         args = [gpf_script, str(self.uproj_path)]
 
