Config file made read only
also any String setting has gotten that option now as well.

diff --git a/Engine/Plugins/VirtualProduction/Switchboard/Source/Switchboard/switchboard/config.py b/Engine/Plugins/VirtualProduction/Switchboard/Source/Switchboard/switchboard/config.py
--- a/Engine/Plugins/VirtualProduction/Switchboard/Source/Switchboard/switchboard/config.py
+++ b/Engine/Plugins/VirtualProduction/Switchboard/Source/Switchboard/switchboard/config.py
@@ -435,6 +435,34 @@ class IntSetting(Setting):
     '''
     A UI-displayable Setting for storing and modifying an integer value.
     '''
+    
+    def __init__(
+        self,
+        attr_name: str,
+        nice_name: str,
+        value: str,
+        tool_tip: typing.Optional[str] = None,
+        show_ui: bool = True,
+        allow_reset: bool = True,
+        migrate_data: typing.Callable[[any], None] = None,
+        is_read_only: bool = False
+    ):
+        '''
+        Create a new IntSetting object.
+
+        Args:
+            attr_name       : Internal name.
+            nice_name       : Display name.
+            value           : The initial value of this Setting.
+            tool_tip        : Tooltip to show in the UI for this Setting.
+            show_ui         : Whether to show this Setting in the Settings UI.
+            is_read_only    : Whether to make entry field editable or not.
+        '''
+        super().__init__(
+            attr_name, nice_name, value,
+            tool_tip=tool_tip, show_ui=show_ui, allow_reset=allow_reset, migrate_data=migrate_data)
+
+        self.is_read_only = is_read_only
 
     def _create_widgets(
             self, override_device_name: typing.Optional[str] = None) \
@@ -447,6 +475,7 @@ class IntSetting(Setting):
         value = str(self.get_value(override_device_name))
         line_edit.setText(value)
         line_edit.setCursorPosition(0)
+        line_edit.setReadOnly(self.is_read_only)
 
         self.set_widget(
             widget=line_edit, override_device_name=override_device_name)
@@ -487,7 +516,8 @@ class StringSetting(Setting):
         tool_tip: typing.Optional[str] = None,
         show_ui: bool = True,
         allow_reset: bool = True,
-        migrate_data: typing.Callable[[any], None] = None
+        migrate_data: typing.Callable[[any], None] = None,
+        is_read_only: bool = False
     ):
         '''
         Create a new StringSetting object.
@@ -499,12 +529,14 @@ class StringSetting(Setting):
             placeholder_text: Placeholder for this Setting's value in the UI.
             tool_tip        : Tooltip to show in the UI for this Setting.
             show_ui         : Whether to show this Setting in the Settings UI.
+            is_read_only    : Whether to make entry field editable or not.
         '''
         super().__init__(
             attr_name, nice_name, value,
             tool_tip=tool_tip, show_ui=show_ui, allow_reset=allow_reset, migrate_data=migrate_data)
 
         self.placeholder_text = placeholder_text
+        self.is_read_only = is_read_only
 
     def _create_widgets(
             self, override_device_name: typing.Optional[str] = None) \
@@ -517,6 +549,7 @@ class StringSetting(Setting):
         line_edit.setText(value)
         line_edit.setPlaceholderText(self.placeholder_text)
         line_edit.setCursorPosition(0)
+        line_edit.setReadOnly(self.is_read_only)
 
         self.set_widget(
             widget=line_edit, override_device_name=override_device_name)
diff --git a/Engine/Plugins/VirtualProduction/Switchboard/Source/Switchboard/switchboard/devices/ndisplay/plugin_ndisplay.py b/Engine/Plugins/VirtualProduction/Switchboard/Source/Switchboard/switchboard/devices/ndisplay/plugin_ndisplay.py
--- a/Engine/Plugins/VirtualProduction/Switchboard/Source/Switchboard/switchboard/devices/ndisplay/plugin_ndisplay.py
+++ b/Engine/Plugins/VirtualProduction/Switchboard/Source/Switchboard/switchboard/devices/ndisplay/plugin_ndisplay.py
@@ -353,12 +353,13 @@ class DevicenDisplay(DeviceUnreal):
     add_device_dialog = AddnDisplayDialog
 
     csettings = {
-        'ndisplay_config_file': FilePathSetting(
+        'ndisplay_config_file': StringSetting(
             attr_name="ndisplay_cfg_file",
             nice_name="nDisplay Config File",
             value="",
-            file_path_filter="nDisplay Config (*.ndisplay;*.uasset)",
-            tool_tip="Path to nDisplay config file"
+            tool_tip="Path to nDisplay config file",
+            allow_reset=False,
+            is_read_only=True
         ),
         'use_all_available_cores': BoolSetting(
             attr_name="use_all_available_cores",
@@ -533,7 +534,8 @@ class DevicenDisplay(DeviceUnreal):
                 attr_name="ue_command_line",
                 nice_name="UE Command Line",
                 value=kwargs.get("ue_command_line", ''),
-                allow_reset=False
+                allow_reset=False,
+                is_read_only=True
             ),
             'window_position': Setting(
                 attr_name="window_position",
