Switchboard: Disable virtual shadow map cache until it supports multi-viewfamily.

diff --git a/Engine/Plugins/VirtualProduction/Switchboard/Source/Switchboard/switchboard/devices/ndisplay/plugin_ndisplay.py b/Engine/Plugins/VirtualProduction/Switchboard/Source/Switchboard/switchboard/devices/ndisplay/plugin_ndisplay.py
--- a/Engine/Plugins/VirtualProduction/Switchboard/Source/Switchboard/switchboard/devices/ndisplay/plugin_ndisplay.py
+++ b/Engine/Plugins/VirtualProduction/Switchboard/Source/Switchboard/switchboard/devices/ndisplay/plugin_ndisplay.py
@@ -899,6 +899,9 @@ class DevicenDisplay(DeviceUnreal):
         # Always tell Chaos to be deterministic
         dp_cvars.append('p.Chaos.Solver.Deterministic=1')
 
+        # Always disable virtual shadow map caching until it supports multi-viewfamily.
+        dp_cvars.append('r.Shadow.Virtual.Cache=0')
+
         # Add user set dp cvars, overriding any of the forced ones.
         user_dp_cvars = self.csettings['ndisplay_dp_cvars'].get_value(self.name)
         user_dp_cvars = [cvar.strip() for cvar in user_dp_cvars if len(cvar.strip()) and len(cvar.split('=')) == 2]
