Searchable ComboBox addition

diff --git a/Engine/Plugins/VirtualProduction/Switchboard/Source/Switchboard/switchboard/devices/ndisplay/plugin_ndisplay.py b/Engine/Plugins/VirtualProduction/Switchboard/Source/Switchboard/switchboard/devices/ndisplay/plugin_ndisplay.py
--- a/Engine/Plugins/VirtualProduction/Switchboard/Source/Switchboard/switchboard/devices/ndisplay/plugin_ndisplay.py
+++ b/Engine/Plugins/VirtualProduction/Switchboard/Source/Switchboard/switchboard/devices/ndisplay/plugin_ndisplay.py
@@ -54,7 +54,7 @@ class AddnDisplayDialog(AddDeviceDialog):
         self.btnFindConfigs.clicked.connect(self.on_clicked_btnFindConfigs)
 
         # Combobox with config files
-        self.cbConfigs = QtWidgets.QComboBox(self)
+        self.cbConfigs = sb_widgets.SearchableComboBox(self)
         self.cbConfigs.setSizePolicy(
             QtWidgets.QSizePolicy.Expanding, QtWidgets.QSizePolicy.Preferred)
         self.cbConfigs.setEditable(True)  # to allow the user to type the value
diff --git a/Engine/Plugins/VirtualProduction/Switchboard/Source/Switchboard/switchboard/switchboard_dialog.py b/Engine/Plugins/VirtualProduction/Switchboard/Source/Switchboard/switchboard/switchboard_dialog.py
--- a/Engine/Plugins/VirtualProduction/Switchboard/Source/Switchboard/switchboard/switchboard_dialog.py
+++ b/Engine/Plugins/VirtualProduction/Switchboard/Source/Switchboard/switchboard/switchboard_dialog.py
@@ -257,6 +257,19 @@ class SwitchboardDialog(QtCore.QObject):
         #self.transport_queue.signal_transport_queue_job_started.connect(self.transport_queue_job_started)
         #self.transport_queue.signal_transport_queue_job_finished.connect(self.transport_queue_job_finished)
 
+        # add level picker combo box and refresh button
+        self.level_combo_box = sb_widgets.SearchableComboBox(self.window)
+        self.level_combo_box.setSizePolicy(QtWidgets.QSizePolicy.Expanding, QtWidgets.QSizePolicy.Fixed)
+        
+        self.refresh_levels_button = sb_widgets.ControlQPushButton()
+        self.refresh_levels_button.setMaximumSize(22, 22)
+        self.refresh_levels_button.setIcon(QtGui.QIcon("icon_refresh.png"))
+        self.refresh_levels_button.setProperty("frameless", True)
+        self.refresh_levels_button.setToolTip("Refresh level list")
+        
+        self.window.horizontalLayout_7.addWidget(self.level_combo_box)
+        self.window.horizontalLayout_7.addWidget(self.refresh_levels_button)
+        
         self.shoot = 'Default'
         self.sequence = SETTINGS.CURRENT_SEQUENCE
         self.slate = SETTINGS.CURRENT_SLATE
@@ -315,8 +328,8 @@ class SwitchboardDialog(QtCore.QObject):
         self.window.slate_line_edit.textChanged.connect(self._set_slate)
         self.window.take_spin_box.valueChanged.connect(self._set_take)
         self.window.sequence_line_edit.textChanged.connect(self._set_sequence)
-        self.window.level_combo_box.currentIndexChanged.connect(self._on_selected_level_changed)
-        self.window.refresh_levels_button.clicked.connect(self.refresh_levels_incremental)
+        self.level_combo_box.currentIndexChanged.connect(self._on_selected_level_changed)
+        self.refresh_levels_button.clicked.connect(self.refresh_levels_incremental)
         self.window.project_cl_combo_box.currentTextChanged.connect(self._set_project_changelist)
         self.window.engine_cl_combo_box.currentIndexChanged.connect(
             lambda _: self._set_engine_changelist(self.window.engine_cl_combo_box.currentText()))
@@ -357,7 +370,7 @@ class SwitchboardDialog(QtCore.QObject):
             btn.setObjectName(name)
             btn.hover_focus = False
 
-        configure_ctrl_btn(self.window.refresh_levels_button, 'refresh')
+        configure_ctrl_btn(self.refresh_levels_button, 'refresh')
         configure_ctrl_btn(self.window.sync_all_button, 'sync')
         configure_ctrl_btn(self.window.build_all_button, 'build')
         configure_ctrl_btn(self.window.sync_and_build_all_button, 'sync_and_build')
@@ -1273,8 +1286,8 @@ class SwitchboardDialog(QtCore.QObject):
         self._set_level(full_map_path)
         
     def _get_level_from_combo_box(self, index: int):
-        # Tooltip stores full path
-        return self.window.level_combo_box.itemData(index, QtCore.Qt.ToolTipRole)
+        # Data stores full path
+        return self.level_combo_box.itemData(index)
 
     def _set_level(self, value):
         ''' Called when level dropdown text changes
@@ -1285,12 +1298,12 @@ class SwitchboardDialog(QtCore.QObject):
             CONFIG.CURRENT_LEVEL = value
             CONFIG.save()
 
-        if self.window.level_combo_box.currentText() != self._level:
-            for index in range(self.window.level_combo_box.count()):
+        if self.level_combo_box.currentText() != self._level:
+            for index in range(self.level_combo_box.count()):
                 if self._get_level_from_combo_box(index) == self._level:
-                    self.window.level_combo_box.blockSignals(True)
-                    self.window.level_combo_box.setCurrentIndex(index)
-                    self.window.level_combo_box.blockSignals(False)
+                    self.level_combo_box.blockSignals(True)
+                    self.level_combo_box.setCurrentIndex(index)
+                    self.level_combo_box.blockSignals(False)
                     break
 
     @property
@@ -1676,8 +1689,8 @@ class SwitchboardDialog(QtCore.QObject):
 
         current_level = CONFIG.CURRENT_LEVEL
 
-        self.window.level_combo_box.clear()
-        self._update_level_list(self.window.level_combo_box, levels)
+        self.level_combo_box.clear()
+        self._update_level_list(self.level_combo_box, levels)
 
         if current_level and current_level in levels:
             self.level = current_level
@@ -1710,7 +1723,7 @@ class SwitchboardDialog(QtCore.QObject):
             
         return short_name_list, short_name_to_path
 
-    def _update_level_list(self, level_combo_box: QtWidgets.QComboBox, level_path_list: List[str]):
+    def _update_level_list(self, level_combo_box: sb_widgets.SearchableComboBox, level_path_list: List[str]):
         def compare_file_names(path_a: str, path_b: str):
             return -1 if path_a.lower() < path_b.lower() \
                 else 1 if path_a.lower() > path_b.lower() else 0
@@ -1725,13 +1738,15 @@ class SwitchboardDialog(QtCore.QObject):
             return
         
         # To disambiguate, show the full path name in the drop-down
+        level_combo_box.setItemData(0, "Default level")
         level_combo_box.setItemData(0, "Default level", QtCore.Qt.ToolTipRole)
         for short_name_index in range(len(short_name_list)):
             short_name = short_name_list[short_name_index]
+            level_combo_box.setItemData(short_name_index + 1, short_name_to_path[short_name])
             level_combo_box.setItemData(short_name_index + 1, short_name_to_path[short_name], QtCore.Qt.ToolTipRole)
 
     def get_current_level_list(self):
-        level_combo = self.window.level_combo_box
+        level_combo = self.level_combo_box
         return [self._get_level_from_combo_box(i) for i in range(1, level_combo.count())] # skip DEFAULT_MAP_TEXT
 
     def refresh_levels_incremental(self):
diff --git a/Engine/Plugins/VirtualProduction/Switchboard/Source/Switchboard/switchboard/switchboard_widgets.py b/Engine/Plugins/VirtualProduction/Switchboard/Source/Switchboard/switchboard/switchboard_widgets.py
--- a/Engine/Plugins/VirtualProduction/Switchboard/Source/Switchboard/switchboard/switchboard_widgets.py
+++ b/Engine/Plugins/VirtualProduction/Switchboard/Source/Switchboard/switchboard/switchboard_widgets.py
@@ -338,6 +338,57 @@ class CollapsibleGroupBox(QtWidgets.QGroupBox):
             self.setMaximumHeight(self.fontMetrics().height() + safety_margin)
             
 
+class SearchableComboBox(QtWidgets.QComboBox):
+    
+    class CustomQList(QtWidgets.QListWidget):
+        def __init__(self, parent):
+            super().__init__(parent)
+            
+            
+    def __init__(self, parent):
+        super().__init__(parent)
+        
+        self.HEIGHT = 15
+        
+        self.item_list = self.CustomQList(self)
+        self.item_list.setViewportMargins(0, self.HEIGHT, 0, 0)
+        self.setView(self.item_list)
+        self.setModel(self.item_list.model())
+        self.resize(300,self.size().height())
+        self.search_line = QtWidgets.QLineEdit(self.view())
+        self.search_line.setVisible(False)
+        self.search_line.setPlaceholderText("Search")
+       
+        self.search_line.textChanged.connect(self.text_changed)
+        
+    def showPopup(self):
+            
+        super().showPopup()
+        
+        MARGIN = 10
+        XPOS = MARGIN/2
+        YPOS = 2
+        
+        self.search_line.setGeometry(XPOS, YPOS, self.view().width() - MARGIN*2.5, self.HEIGHT)
+        self.search_line.setVisible(True)
+        
+        self.CB_item_text = [self.itemText(i) for i in range(self.count())]
+        
+    def hidePopup(self):
+        super().hidePopup()
+        
+        self.search_line.setVisible(False)
+        self.search_line.clear()
+    
+    def text_changed(self, text:str):
+        
+        for item in self.CB_item_text:
+            if item.lower().__contains__(text.lower()):
+                self.view().setRowHidden(self.findText(item), False)
+            else:
+                self.view().setRowHidden(self.findText(item), True)
+     
+
 def set_qt_property(
     widget: QtWidgets.QWidget, prop, value, *,
     update_box_model: bool = True,
