Take number auto increments now.

diff --git a/Engine/Plugins/VirtualProduction/Switchboard/Source/Switchboard/switchboard/switchboard_dialog.py b/Engine/Plugins/VirtualProduction/Switchboard/Source/Switchboard/switchboard/switchboard_dialog.py
--- a/Engine/Plugins/VirtualProduction/Switchboard/Source/Switchboard/switchboard/switchboard_dialog.py
+++ b/Engine/Plugins/VirtualProduction/Switchboard/Source/Switchboard/switchboard/switchboard_dialog.py
@@ -1825,6 +1825,7 @@ class SwitchboardDialog(QtCore.QObject):
         if self._is_recording:
             LOGGER.debug('Record stop button pressed')
             self._record_stop(1)
+            self.window.take_spin_box.setValue(self.window.take_spin_box.value() + 1)
         else:
             LOGGER.debug('Record start button pressed')
             self._record_start(self.slate, self.take, self.description)
