Switchboard: Added "Create Config" menu item.

This automates the creation of a Switchboard config from the current map and nDisplay actor, automating this per-user drudgery. SB and SBL will be auto-launched and the config will be created and opened automatically. The user can opt to have the devices auto-connect, or also have all the ip addresses be overridden with localhost.

The user can also pick the nDisplay instance to use (when there are more than one), and select how many Editor devices he wants created. With was done w/o adding any plugin dependencies.

In addition:
* Configs can now be "Saved As" and loaded to/from anywhere on disk.
* SBL now auto-builds if necessary when launching from the Editor
* Users can now pass python scripts and arbitrary arguments for them using the --script and --scriptargs cmd line arguments, and if the script contains a subclass of SwitchboardScriptBase, its callbacks will be called at certain times of SB lifetime. In particular, during initialization. This enables certain externally driven automation of SB, and is how the automated creation of new configs is done by Unreal.

diff --git a/Engine/Plugins/VirtualProduction/Switchboard/Source/Switchboard/switchboard.bat b/Engine/Plugins/VirtualProduction/Switchboard/Source/Switchboard/switchboard.bat
--- a/Engine/Plugins/VirtualProduction/Switchboard/Source/Switchboard/switchboard.bat
+++ b/Engine/Plugins/VirtualProduction/Switchboard/Source/Switchboard/switchboard.bat
@@ -1,4 +1,5 @@
 @echo off
+
 setlocal ENABLEDELAYEDEXPANSION
 
 REM This script optionally takes a single argument, representing the path to the desired Python
@@ -8,6 +9,7 @@ REM Additionally, in the case where no argument was provided, pause on errors, i
 REM double-clicked the batch file and we don't want the console window to vanish inexplicably.
 REM Otherwise, we exit immediately.
 
+set _Args=%*
 
 set _switchboardDir=%~dp0
 REM pushd and %CD% are used to normalize the relative path to a shorter absolute path.
@@ -17,16 +19,22 @@ set _enginePythonPlatformDir=%_engineDir%\Binaries\ThirdParty\Python3\Win64
 set _defaultVenvDir=%_engineDir%\Extras\ThirdPartyNotUE\SwitchboardThirdParty\Python
 popd
 
-if not "%~1"=="" (
+if "%~1"=="" goto:defaultenv
+if "%~1"=="--defaultenv" goto:defaultenv
+
     set _bUsingDefaultVenv=
     set _venvDir=%~1
     echo Using provided path for Python virtual environment (!_venvDir!^)
-) else (
-    REM Called with no arguments; assume user double-clicked on us and pause for errors.
+
+    goto:fi_env
+
+:defaultenv
+
     set _bUsingDefaultVenv=1
     set _venvDir=%_defaultVenvDir%
     echo Using DEFAULT path for Python virtual environment (%_defaultVenvDir%^)
-)
+
+:fi_env
 
 call:main
 
@@ -54,7 +62,8 @@ goto:eof
 :start_sb
 
 set PYTHONPATH=%_switchboardDir%;%PYTHONPATH%
-start "Switchboard" /D "%_switchboardDir%" "%_venvDir%\Scripts\pythonw.exe" -m switchboard
+
+start "Switchboard" /D "%_switchboardDir%" "%_venvDir%\Scripts\pythonw.exe" -m switchboard %_Args%
 if %ERRORLEVEL% NEQ 0 (
     echo Failed to launch Switchboard^!
     if defined _bUsingDefaultVenv ( pause )
diff --git a/Engine/Plugins/VirtualProduction/Switchboard/Source/Switchboard/switchboard.sh b/Engine/Plugins/VirtualProduction/Switchboard/Source/Switchboard/switchboard.sh
--- a/Engine/Plugins/VirtualProduction/Switchboard/Source/Switchboard/switchboard.sh
+++ b/Engine/Plugins/VirtualProduction/Switchboard/Source/Switchboard/switchboard.sh
@@ -3,12 +3,13 @@
 # This script optionally takes a single argument, representing the path to the desired Python
 # virtual environment directory. If omitted, it defaults to the value of $_defaultVenvDir.
 
+_Args=$@
 _switchboardDir=$(dirname "$0")
 _engineDir=$(cd "$_switchboardDir/../../../../.."; pwd)
 _enginePythonDir="$_engineDir/Binaries/ThirdParty/Python3/Linux"
 _defaultVenvDir="$_engineDir/Extras/ThirdPartyNotUE/SwitchboardThirdParty/Python"
 
-if [ $# -lt 1 ]; then
+if [ $# -lt 1 ] || [ $1 == "--defaultenv" ]; then
     _venvDir=$_defaultVenvDir
 else
     _venvDir="$1"
@@ -23,4 +24,4 @@ if [ ! -x "$_venvDir/bin/python3" ]; then
     fi
 fi
 
-(cd "$_switchboardDir" && PYTHONPATH="$_switchboardDir:$PYTHONPATH" "$_venvDir/bin/python3" -m switchboard)
+(cd "$_switchboardDir" && PYTHONPATH="$_switchboardDir:$PYTHONPATH" "$_venvDir/bin/python3" -m switchboard $_Args)
diff --git a/Engine/Plugins/VirtualProduction/Switchboard/Source/Switchboard/switchboard/__main__.py b/Engine/Plugins/VirtualProduction/Switchboard/Source/Switchboard/switchboard/__main__.py
--- a/Engine/Plugins/VirtualProduction/Switchboard/Source/Switchboard/switchboard/__main__.py
+++ b/Engine/Plugins/VirtualProduction/Switchboard/Source/Switchboard/switchboard/__main__.py
@@ -1,16 +1,26 @@
 # Copyright Epic Games, Inc. All Rights Reserved.
 from switchboard.config import SETTINGS, CONFIG
+from switchboard.switchboard_scripting import ScriptManager
 from .switchboard_dialog import SwitchboardDialog
 from .switchboard_logging import LOGGER
 
 import signal
 import sys
+import argparse
 
 from PySide2 import QtCore, QtWidgets
 
 # Build resources
 # "C:\Program Files (x86)\Python37-32\Lib\site-packages\PySide2\pyside2-rcc" -o D:\Switchboard\switchboard\resources.py D:\Switchboard\switchboard\ui\resources.qrc
 
+def parse_arguments():
+    ''' Parses command line arguments and returns the populated namespace 
+    '''
+    parser = argparse.ArgumentParser()
+    parser.add_argument('--script'    , default='', help='Path to script that contains a SwichboardScriptBase subclass')
+    parser.add_argument('--scriptargs', default='', help='String to pass to SwichboardScriptBase subclass as arguments')
+    args, unknown = parser.parse_known_args()
+    return args
 
 def launch():
     """
@@ -29,7 +39,24 @@ def launch():
     SETTINGS.init()
     CONFIG.init(SETTINGS.CONFIG)
     
-    main_window = SwitchboardDialog()
+    # parse arguments
+    args = parse_arguments()
+
+    # script manager
+    script_manager = ScriptManager()
+
+    # add script passed from command line
+    if args.script:
+        try:
+            script_manager.add_script_from_path(args.script, args.scriptargs)
+        except Exception as e:
+            LOGGER.warning(f"Could not initialize '{args.script}': {e}")
+
+    # script pre-init
+    script_manager.on_preinit()
+
+    # create main window
+    main_window = SwitchboardDialog(script_manager)
 
     if not main_window.window:
         return
@@ -57,8 +84,13 @@ def launch():
     timer.start(200)
     timer.timeout.connect(lambda: None)
 
-    sys.exit(app.exec_())
+    # execute the app
+    appresult = app.exec_()
+
+    # script exit
+    script_manager.on_exit()
 
+    sys.exit(appresult)
 
 if __name__ == "__main__":
     launch()
\ No newline at end of file
diff --git a/Engine/Plugins/VirtualProduction/Switchboard/Source/Switchboard/switchboard/config.py b/Engine/Plugins/VirtualProduction/Switchboard/Source/Switchboard/switchboard/config.py
--- a/Engine/Plugins/VirtualProduction/Switchboard/Source/Switchboard/switchboard/config.py
+++ b/Engine/Plugins/VirtualProduction/Switchboard/Source/Switchboard/switchboard/config.py
@@ -1817,7 +1817,6 @@ def get_absolute_config_path(
     The string/path is validated to ensure that:
       - It is not empty, or all whitespace
       - It ends with the config path suffix
-      - If it is already absolute, that it is underneath the root configs path
       - It is not the same path as the user settings file path
     '''
     if isinstance(config_path, str):
@@ -1834,16 +1833,7 @@ def get_absolute_config_path(
         config_path = config_path.with_name(
             f'{config_path.name}{CONFIG_SUFFIX}')
 
-    if config_path.is_absolute():
-        # Paths that are already absolute must have the root configs path as a
-        # parent path.
-        # Python 3.9 introduced pathlib.Path.is_relative_to(), which would read
-        # a bit nicer here.
-        if ROOT_CONFIGS_PATH not in config_path.parents:
-            raise ConfigPathLocationError(
-                f'Config path "{config_path}" is not underneath the root '
-                f'configs path "{ROOT_CONFIGS_PATH}"')
-    else:
+    if not config_path.is_absolute():
         # Relative paths can simply be made absolute.
         config_path = ROOT_CONFIGS_PATH.joinpath(config_path)
 
@@ -2155,12 +2145,16 @@ class Config(object):
         self.INSIGHTS_TRACE_ARGS = self.unreal_insight_settings["tracing_args"]
         self.INSIGHTS_STAT_EVENTS = self.unreal_insight_settings["tracing_stat_events"]
 
+    def default_mu_server_name(self):
+        ''' Returns default server name based on current settings '''
+        return f'{self.PROJECT_NAME.get_value()}_MU_Server'
+
     def init_muserver(self, data={}):
         self.mu_settings = {
             "muserver_server_name": StringSetting(
                 "muserver_server_name",
                 "Server name",
-                data.get('muserver_server_name', f'{self.PROJECT_NAME.get_value()}_MU_Server'),
+                data.get('muserver_server_name', self.default_mu_server_name()),
                 tool_tip="The name that will be given to the server"
             ),
             "muserver_command_line_arguments": StringSetting(
@@ -2291,7 +2285,7 @@ class Config(object):
 
         if self.file_path:
             new_config_path.parent.mkdir(parents=True, exist_ok=True)
-            self.file_path.replace(new_config_path)
+            shutil.move(self.file_path, new_config_path)
 
         self.file_path = new_config_path
         self.save()
@@ -2580,6 +2574,8 @@ class UserSettings(object):
                 data['config'] = str(get_relative_config_path(self.CONFIG))
             except ConfigPathError as e:
                 LOGGER.error(e)
+            except ValueError as e:
+                data['config'] = str(self.CONFIG)
 
         with open(USER_SETTINGS_FILE_PATH, 'w') as f:
             json.dump(data, f, indent=4)
diff --git a/Engine/Plugins/VirtualProduction/Switchboard/Source/Switchboard/switchboard/switchboard_dialog.py b/Engine/Plugins/VirtualProduction/Switchboard/Source/Switchboard/switchboard/switchboard_dialog.py
--- a/Engine/Plugins/VirtualProduction/Switchboard/Source/Switchboard/switchboard/switchboard_dialog.py
+++ b/Engine/Plugins/VirtualProduction/Switchboard/Source/Switchboard/switchboard/switchboard_dialog.py
@@ -6,7 +6,10 @@ import os
 import threading
 import time
 import re
-from typing import List, Optional, Set
+import shutil
+from typing import List, Optional, Set, Union
+
+from pathlib import Path
 
 from PySide2 import QtCore
 from PySide2 import QtGui
@@ -167,6 +170,7 @@ class DeviceAdditionalSettingsUI(QtCore.QObject):
         return button
 
 class SwitchboardDialog(QtCore.QObject):
+
     STYLESHEET_PATH = os.path.join(RELATIVE_PATH, 'ui/switchboard.qss')
 
     _stylesheet_watcher: Optional[QtCore.QFileSystemWatcher] = None
@@ -188,9 +192,11 @@ class SwitchboardDialog(QtCore.QObject):
             stylesheet = styling.read()
             QtWidgets.QApplication.instance().setStyleSheet(stylesheet)
 
-    def __init__(self):
+    def __init__(self, script_manager):
         super().__init__()
 
+        self.script_manager = script_manager
+
         font_dir = os.path.join(ENGINE_PATH, 'Content/Slate/Fonts')
         font_files = ['Roboto-Regular.ttf', 'Roboto-Bold.ttf',
                       'DroidSansMono.ttf']
@@ -394,6 +400,7 @@ class SwitchboardDialog(QtCore.QObject):
 
         # Menu items
         self.window.menu_new_config.triggered.connect(self.menu_new_config)
+        self.window.menu_save_config_as.triggered.connect(self.menu_save_config_as)
         self.window.menu_delete_config.triggered.connect(self.menu_delete_config)
         self.window.update_settings.triggered.connect(self.menu_update_settings)
 
@@ -404,12 +411,11 @@ class SwitchboardDialog(QtCore.QObject):
         # Plugin UI
         self.device_manager.plug_into_ui(self.window.menu_bar, self.window.tabs_main)
 
-        # If starting up with new config, open the menu to create a new one
-        if not CONFIG.file_path:
-            self.menu_new_config()
-        else:
+        if CONFIG.file_path:
             self.toggle_p4_controls(CONFIG.P4_ENABLED.get_value())
             self.refresh_levels()
+        else:
+            self.menu_new_config()
 
         self.set_config_hooks()
 
@@ -426,6 +432,8 @@ class SwitchboardDialog(QtCore.QObject):
         self.window.current_ip_value.editingFinished.connect(self._try_change_ip_address)
 
         self.refresh_window_title()
+
+        self.script_manager.on_postinit(self)
         
     def _try_change_ip_address(self):
         def is_valid_ip_format(address):
@@ -756,6 +764,27 @@ class SwitchboardDialog(QtCore.QObject):
             if config_path == SETTINGS.CONFIG:
                 config_action.setEnabled(False)
 
+        # Make a special entry for a config not in the normal area
+        externalconfig_action = QtWidgets.QAction("Browse...", self.window.menu_load_config)
+        externalconfig_action.triggered.connect(self._on_open_external_config)
+        self.window.menu_load_config.addAction(externalconfig_action)
+
+    def _on_open_external_config(self):
+        ''' When the user wants to open a config located outside of the designated area
+        This is useful when the user keeps the configs somewhere else (e.g. under the project)
+        '''
+        config_path, _ = QtWidgets.QFileDialog.getOpenFileName(self.window,
+            'Select config file', str(config.ROOT_CONFIGS_PATH),
+            f'Config files (*{config.CONFIG_SUFFIX})'
+        )
+
+        # Do nothing if the user didn't choose a config path
+        if not config_path:
+            return
+
+        # ok, let's open it
+        self.set_current_config(config_path)
+
     def set_current_config(self, config_path):
 
         # Update to the new config
@@ -813,6 +842,84 @@ class SwitchboardDialog(QtCore.QObject):
     def update_current_ip_address_text(self):
         self.window.current_ip_value.setText(SETTINGS.IP_ADDRESS.get_value())
 
+    def create_new_config(self, file_path: Union[str, Path], uproject, engine_dir, p4_settings):
+        ''' Creates a new config file
+
+        Args:
+            file_path: Path to the new config file
+            uproject: Path to the unreal project (.uproject) file
+            engine_dir: Path to the Engine/ directory
+            p4_settings: Desired Perforce settings
+        '''
+
+        CONFIG.init_new_config(
+            file_path=file_path,
+            uproject=uproject,
+            engine_dir=engine_dir,
+            p4_settings=p4_settings
+        )
+
+        # Disable saving while loading
+        CONFIG.push_saving_allowed(False)
+        try:
+            # Remove all devices
+            self.device_manager.clear_device_list()
+            self.device_list_widget.clear_widgets()
+
+            # Reset plugin settings
+            self.device_manager.reset_plugins_settings(CONFIG)
+
+            # Set hooks to this dialog's UI
+            self.set_config_hooks()
+        finally:
+            # Re-enable saving after loading
+            CONFIG.pop_saving_allowed()
+
+        # Update the UI
+        self.toggle_p4_controls(CONFIG.P4_ENABLED.get_value())
+        self.refresh_levels()
+        self.update_current_config_text()
+        self.refresh_muserver_autojoin()
+        self.refresh_trace_settings()
+        self.refresh_window_title()
+
+    def menu_save_config_as(self):
+        ''' Copy the current config file and move to it '''
+
+        # get the destination path for the copy
+
+        new_config_path, _ = QtWidgets.QFileDialog.getSaveFileName(self.window,
+            'Select config file', str(config.ROOT_CONFIGS_PATH),
+            f'Config files (*{config.CONFIG_SUFFIX})',
+            #options=QtWidgets.QFileDialog.DontUseNativeDialog
+        )
+
+        # Do nothing if the user didn't choose a destination config path
+        if not new_config_path:
+            return
+
+        # If replacing the same config, just save
+        if Path(new_config_path).resolve() == Path(SETTINGS.CONFIG).resolve():
+            CONFIG.save()
+            return
+
+        # Switch to the new config
+        CONFIG.replace(new_config_path)
+
+        # Update the settings
+        SETTINGS.CONFIG = new_config_path
+        SETTINGS.save()
+
+        # Make sure we reflect the new config in the status bar
+        self.update_current_config_text()
+
+        # Reset mu server name
+        # * unless it is running to avoid devices failing to connect to the current server
+        # * if there are devices running, they won't auto-connect when you launch the renamed server
+        mu_server = switchboard_application.get_multi_user_server_instance()
+        if not mu_server.is_running():
+            CONFIG.MUSERVER_SERVER_NAME.update_value(CONFIG.default_mu_server_name())
+
     def menu_new_config(self):
         uproject_search_path = os.path.dirname(CONFIG.UPROJECT_PATH.get_value().replace('"',''))
 
@@ -827,33 +934,12 @@ class SwitchboardDialog(QtCore.QObject):
 
         if dialog.result() == QtWidgets.QDialog.Accepted:
 
-            CONFIG.init_new_config(file_path=dialog.config_path,
-                uproject=dialog.uproject, engine_dir=dialog.engine_dir,
-                p4_settings=dialog.p4_settings())
-
-            # Disable saving while loading
-            CONFIG.push_saving_allowed(False)
-            try:
-                # Remove all devices
-                self.device_manager.clear_device_list()
-                self.device_list_widget.clear_widgets()
-
-                # Reset plugin settings
-                self.device_manager.reset_plugins_settings(CONFIG)
-
-                # Set hooks to this dialog's UI
-                self.set_config_hooks()
-            finally:
-                # Re-enable saving after loading
-                CONFIG.pop_saving_allowed()
-
-            # Update the UI
-            self.toggle_p4_controls(CONFIG.P4_ENABLED.get_value())
-            self.refresh_levels()
-            self.update_current_config_text()
-            self.refresh_muserver_autojoin()
-            self.refresh_trace_settings()
-            self.refresh_window_title()
+            self.create_new_config(
+                file_path=dialog.config_path,
+                uproject=dialog.uproject, 
+                engine_dir=dialog.engine_dir,
+                p4_settings=dialog.p4_settings()
+            )
 
     def menu_delete_config(self):
         """
@@ -900,7 +986,7 @@ class SwitchboardDialog(QtCore.QObject):
         settings_dialog.select_all_tab()
 
         old_ip_address = SETTINGS.IP_ADDRESS.get_value()
-        old_transport_path = SETTINGS.TRANSPORT_PATH.get_value()
+
         # avoid saving the config all the time while in the settings dialog
         CONFIG.push_saving_allowed(False)
         try:
diff --git a/Engine/Plugins/VirtualProduction/Switchboard/Source/Switchboard/switchboard/switchboard_scripting.py b/Engine/Plugins/VirtualProduction/Switchboard/Source/Switchboard/switchboard/switchboard_scripting.py
new file mode 100644
--- /dev/null
+++ b/Engine/Plugins/VirtualProduction/Switchboard/Source/Switchboard/switchboard/switchboard_scripting.py
@@ -0,0 +1,93 @@
+# Copyright Epic Games, Inc. All Rights Reserved.
+
+import sys, importlib, inspect
+from pathlib import Path
+
+from switchboard.switchboard_logging import LOGGER
+
+
+class SwitchboardScriptBase:
+    ''' Base class for script files that can be passed in the command line using -script and -scriptargs'''
+
+    scriptargs = ''
+
+    def __init__(self, scriptargs):
+        self.scriptargs = scriptargs
+
+    def on_preinit(self):
+        ''' Called before SB UI initializes'''
+        pass
+
+    def on_postinit(self, app):
+        ''' Called after SB UI initializes'''
+        pass
+
+    def on_exit(self):
+        ''' Called right before the application exits'''
+        pass
+
+class ScriptManager:
+
+    scripts = []
+
+    def add_script_from_path(self, scriptpath, scriptargs):
+        
+        # file existence sanity check
+        path = Path(scriptpath)
+
+        if not path.is_file():
+            raise FileNotFoundError
+
+        # insert its folder to sys path to import it
+        script_directory = str(path.parents[0].absolute())
+        sys.path.insert(0, script_directory)
+
+        # import the module by name
+        scriptmod = importlib.import_module(path.stem)
+
+        # find SwitchboardScriptBase subclasses. 
+        # They will all share the same scriptargs
+        for name, scriptclass in inspect.getmembers(scriptmod):
+            if inspect.isclass(scriptclass) and scriptclass in SwitchboardScriptBase.__subclasses__():
+
+                try:
+                    script = scriptclass(scriptargs)
+                except Exception as e:
+                    LOGGER.error(f"Could not initialize script '{name}': {e}")
+                    raise
+
+                self.scripts.append(script)
+
+    def on_preinit(self):
+
+        for script in self.scripts:
+            try:
+                script.on_preinit()
+            except Exception as e:
+                LOGGER.error(e)
+
+    def on_postinit(self, app):
+
+        for script in self.scripts:
+            try:
+                script.on_postinit(app)
+            except Exception as e:
+                LOGGER.error(e)
+                
+    def on_exit(self):
+
+        for script in self.scripts:
+            try:
+                script.on_exit()
+            except Exception as e:
+                LOGGER.error(e)
+                
+
+
+
+
+        
+
+        
+
+    
\ No newline at end of file
diff --git a/Engine/Plugins/VirtualProduction/Switchboard/Source/SwitchboardEditor/Private/DisplayClusterRootActorReferenceDetailCustomization.cpp b/Engine/Plugins/VirtualProduction/Switchboard/Source/SwitchboardEditor/Private/DisplayClusterRootActorReferenceDetailCustomization.cpp
new file mode 100644
--- /dev/null
+++ b/Engine/Plugins/VirtualProduction/Switchboard/Source/SwitchboardEditor/Private/DisplayClusterRootActorReferenceDetailCustomization.cpp
@@ -0,0 +1,119 @@
+// Copyright Epic Games, Inc. All Rights Reserved.
+
+#include "DisplayClusterRootActorReferenceDetailCustomization.h"
+
+#include "DetailLayoutBuilder.h"
+#include "DetailWidgetRow.h"
+#include "IPropertyUtilities.h"
+#include "IPropertyTypeCustomization.h"
+#include "PropertyCustomizationHelpers.h"
+#include "PropertyHandle.h"
+#include "SwitchboardEditorModule.h"
+#include "SwitchboardTypes.h"
+
+
+#define LOCTEXT_NAMESPACE "DisplayClusterRootActorReferenceDetailCustomization"
+
+
+
+TSharedRef<IPropertyTypeCustomization> FDisplayClusterRootActorReferenceDetailCustomization::MakeInstance()
+{
+	return MakeShareable(new FDisplayClusterRootActorReferenceDetailCustomization);
+}
+
+FDisplayClusterRootActorReference* FDisplayClusterRootActorReferenceDetailCustomization::GetDisplayClusterRootActorReference() const
+{
+	// Access the underlying structure
+	TArray<void*> RawData;
+	StructPropertyHandle->AccessRawData(RawData);
+
+	if (RawData.Num() != 1)
+	{
+		return nullptr;
+	}
+
+	return reinterpret_cast<FDisplayClusterRootActorReference*>(RawData[0]);
+}
+
+void FDisplayClusterRootActorReferenceDetailCustomization::CustomizeHeader(TSharedRef<IPropertyHandle> InPropertyHandle, FDetailWidgetRow& HeaderRow, IPropertyTypeCustomizationUtils& CustomizationUtils)
+{
+	// Early return if the nDisplay actor class cannot be found. This property will not be displayed.
+	if (!FSwitchboardEditorModule::GetDisplayClusterRootActorClass())
+	{
+		HeaderRow.Visibility(EVisibility::Collapsed);
+		return;
+	}
+
+	// Keep our property handle so that the widget can access it.
+	StructPropertyHandle = InPropertyHandle;
+
+	// Make sure the property is of the expected type
+	check(CastFieldChecked<FStructProperty>(StructPropertyHandle->GetProperty())->Struct == FDisplayClusterRootActorReference::StaticStruct());
+
+	// Access the underlying structure
+	TArray<void*> RawData;
+	StructPropertyHandle->AccessRawData(RawData);
+
+	check(RawData.Num() == 1);
+	FDisplayClusterRootActorReference* DCRAReference = reinterpret_cast<FDisplayClusterRootActorReference*>(RawData[0]);
+
+	// Customize the widget
+
+	TSharedPtr<IPropertyUtilities> PropertyUtils = CustomizationUtils.GetPropertyUtilities();
+
+	HeaderRow.NameContent()
+	[
+		// Just the name
+		StructPropertyHandle->CreatePropertyNameWidget(LOCTEXT("nDisplayActor", "nDisplay Actor"))
+	]
+	.ValueContent()
+	.MinDesiredWidth(300)
+	[
+		// nDisplay actor picker
+
+		SNew(SHorizontalBox)
+		+ SHorizontalBox::Slot()
+		.FillWidth(1.0f)
+		.AutoWidth()
+		.VAlign(VAlign_Center)
+		[
+			SNew(SObjectPropertyEntryBox)
+			.DisplayBrowse(false)
+			.AllowedClass(FSwitchboardEditorModule::GetDisplayClusterRootActorClass())
+			.OnObjectChanged_Lambda([&](const FAssetData& AssetData) -> void
+			{
+				FDisplayClusterRootActorReference* DCRAReference = GetDisplayClusterRootActorReference();
+
+				if (!DCRAReference)
+				{
+					return;
+				}
+
+				if (AssetData.IsValid())
+				{
+					DCRAReference->DCRA = Cast<AActor>(AssetData.GetAsset());
+				}
+				else
+				{
+					DCRAReference->DCRA.Reset();
+				}
+			})
+			.ObjectPath_Lambda([&]() -> FString
+			{
+				FDisplayClusterRootActorReference* DCRAReference = GetDisplayClusterRootActorReference();
+
+				if (DCRAReference && DCRAReference->DCRA.IsValid())
+				{
+					FAssetData AssetData(DCRAReference->DCRA.Get(), true);
+					return AssetData.ObjectPath.ToString();
+				}
+
+				return TEXT("");
+			})
+		]
+	].IsEnabled(MakeAttributeLambda([=] { return !InPropertyHandle->IsEditConst() && PropertyUtils->IsPropertyEditingEnabled(); }));
+}
+
+
+
+#undef LOCTEXT_NAMESPACE
\ No newline at end of file
diff --git a/Engine/Plugins/VirtualProduction/Switchboard/Source/SwitchboardEditor/Private/DisplayClusterRootActorReferenceDetailCustomization.h b/Engine/Plugins/VirtualProduction/Switchboard/Source/SwitchboardEditor/Private/DisplayClusterRootActorReferenceDetailCustomization.h
new file mode 100644
--- /dev/null
+++ b/Engine/Plugins/VirtualProduction/Switchboard/Source/SwitchboardEditor/Private/DisplayClusterRootActorReferenceDetailCustomization.h
@@ -0,0 +1,30 @@
+// Copyright Epic Games, Inc. All Rights Reserved.
+
+#pragma once
+
+#include "IPropertyTypeCustomization.h"
+
+#include "SwitchboardTypes.h"
+
+
+class FDisplayClusterRootActorReferenceDetailCustomization : public IPropertyTypeCustomization
+{
+public:
+	/** Makes a new instance of this detail layout class for a specific detail view requesting it. */
+	static TSharedRef<IPropertyTypeCustomization> MakeInstance();
+
+public:
+	// IDetailCustomization interface
+	virtual void CustomizeHeader( TSharedRef<IPropertyHandle> PropertyHandle, FDetailWidgetRow& HeaderRow, IPropertyTypeCustomizationUtils& CustomizationUtils ) override;
+	virtual void CustomizeChildren(TSharedRef<IPropertyHandle> PropertyHandle, IDetailChildrenBuilder& ChildBuilder, IPropertyTypeCustomizationUtils& CustomizationUtils) override {}
+
+private:
+
+	/** Returns the DCRA reference in the cached property handle */
+	FDisplayClusterRootActorReference* GetDisplayClusterRootActorReference() const;
+
+private:
+
+	/** Handle to the structure property we're editing */
+	TSharedPtr<IPropertyHandle> StructPropertyHandle;
+};
diff --git a/Engine/Plugins/VirtualProduction/Switchboard/Source/SwitchboardEditor/Private/SwitchboardEditorModule.cpp b/Engine/Plugins/VirtualProduction/Switchboard/Source/SwitchboardEditor/Private/SwitchboardEditorModule.cpp
--- a/Engine/Plugins/VirtualProduction/Switchboard/Source/SwitchboardEditor/Private/SwitchboardEditorModule.cpp
+++ b/Engine/Plugins/VirtualProduction/Switchboard/Source/SwitchboardEditor/Private/SwitchboardEditorModule.cpp
@@ -1,17 +1,29 @@
 // Copyright Epic Games, Inc. All Rights Reserved.
 
 #include "SwitchboardEditorModule.h"
-#include "SwitchboardEditorSettings.h"
-#include "SwitchboardMenuEntry.h"
-#include "SwitchboardProjectSettings.h"
-#include "SwitchboardSettingsCustomization.h"
+
+#include "AssetRegistry/AssetRegistryModule.h"
 #include "Async/Async.h"
+#include "DesktopPlatformModule.h"
+#include "DisplayClusterRootActorReferenceDetailCustomization.h"
 #include "Editor.h"
 #include "EditorUtilitySubsystem.h"
+#include "Interfaces/IPluginManager.h"
 #include "ISettingsModule.h"
 #include "ISettingsSection.h"
+#include "JsonObjectConverter.h"
 #include "Misc/CoreDelegates.h"
+#include "Misc/FeedbackContext.h"
+#include "Misc/FileHelper.h"
+#include "Misc/ScopedSlowTask.h"
 #include "Misc/ScopeExit.h"
+#include "SwitchboardEditorSettings.h"
+#include "SwitchboardMenuEntry.h"
+#include "SwitchboardProjectSettings.h"
+#include "SwitchboardSettingsCustomization.h"
+#include "SwitchboardTypes.h"
+
+#include <filesystem>
 
 #if PLATFORM_WINDOWS
 #include "Windows/AllowWindowsPlatformTypes.h"
@@ -33,6 +45,8 @@ bool RemoveListenerAutolaunchEntry();
 DEFINE_LOG_CATEGORY(LogSwitchboardPlugin);
 
 
+TSoftObjectPtr<UClass> FSwitchboardEditorModule::DisplayClusterRootActorClass;
+
 // static
 const FString& FSwitchboardEditorModule::GetSbScriptsPath()
 {
@@ -65,6 +79,23 @@ const FString& FSwitchboardEditorModule::GetSbExePath()
 	return ExePath;
 }
 
+// static
+UClass* FSwitchboardEditorModule::GetDisplayClusterRootActorClass()
+{
+	// Searching by name of the class avoids dependency on nDisplay plugin
+
+	// Early return if we have already found it
+	if (DisplayClusterRootActorClass.IsValid())
+	{
+		return DisplayClusterRootActorClass.Get();
+	}
+
+	// Try to find it by name.
+	DisplayClusterRootActorClass = FindObject<UClass>(ANY_PACKAGE, TEXT("DisplayClusterRootActor"));
+
+	return DisplayClusterRootActorClass.Get();
+}
+
 
 void FSwitchboardEditorModule::StartupModule()
 {
@@ -88,8 +119,21 @@ void FSwitchboardEditorModule::StartupModule()
 	}
 
 	FPropertyEditorModule& PropertyModule = FModuleManager::GetModuleChecked<FPropertyEditorModule>("PropertyEditor");
-	PropertyModule.RegisterCustomClassLayout(USwitchboardEditorSettings::StaticClass()->GetFName(), FOnGetDetailCustomizationInstance::CreateStatic(&FSwitchboardEditorSettingsCustomization::MakeInstance));
-	PropertyModule.RegisterCustomClassLayout(USwitchboardProjectSettings::StaticClass()->GetFName(), FOnGetDetailCustomizationInstance::CreateStatic(&FSwitchboardProjectSettingsCustomization::MakeInstance));
+
+	PropertyModule.RegisterCustomClassLayout(
+		USwitchboardEditorSettings::StaticClass()->GetFName(), 
+		FOnGetDetailCustomizationInstance::CreateStatic(&FSwitchboardEditorSettingsCustomization::MakeInstance)
+	);
+
+	PropertyModule.RegisterCustomClassLayout(
+		USwitchboardProjectSettings::StaticClass()->GetFName(), 
+		FOnGetDetailCustomizationInstance::CreateStatic(&FSwitchboardProjectSettingsCustomization::MakeInstance)
+	);
+
+	PropertyModule.RegisterCustomPropertyTypeLayout(
+		FDisplayClusterRootActorReference::StaticStruct()->GetFName(), 
+		FOnGetPropertyTypeCustomizationInstance::CreateStatic(&FDisplayClusterRootActorReferenceDetailCustomization::MakeInstance)
+	);
 
 	DeferredStartDelegateHandle = FCoreDelegates::OnFEngineLoopInitComplete.AddRaw(this, &FSwitchboardEditorModule::OnEngineInitComplete);
 
@@ -118,22 +162,129 @@ void FSwitchboardEditorModule::ShutdownModule()
 }
 
 
-bool FSwitchboardEditorModule::LaunchSwitchboard()
+bool FSwitchboardEditorModule::LaunchSwitchboard(const FString& Arguments)
 {
-	const FString ScriptArgs = FString::Printf(TEXT("\"%s\""),
-		*GetDefault<USwitchboardEditorSettings>()->VirtualEnvironmentPath.Path);
+	const FString ScriptArgs = FString::Printf(TEXT("\"%s\" %s"),
+		*GetDefault<USwitchboardEditorSettings>()->VirtualEnvironmentPath.Path,
+		*Arguments
+	);
 
 	return RunProcess(GetSbExePath(), ScriptArgs);
 }
 
+bool FSwitchboardEditorModule::CompileSwitchboardListener() const
+{
+	IDesktopPlatform* DesktopPlatform = FDesktopPlatformModule::Get();
+
+	if (!DesktopPlatform || !DesktopPlatform->IsUnrealBuildToolAvailable() || DesktopPlatform->IsUnrealBuildToolRunning())
+	{
+		return false;
+	}
+
+	FFeedbackContext* FeedbackContext = GWarn;
+	check(FeedbackContext);
+
+	FScopedSlowTask SlowTask(1.0f, LOCTEXT("BuildingSwitchboardListener", "Building SwitchboardListener..."), true /* bIsEnabled */, *FeedbackContext);
+	SlowTask.MakeDialog(false /*bShowCancelButton*/, false /*bAllowInPie*/);
+
+	const FString Arguments = FString::Printf(TEXT("SwitchboardListener Development %s -Progress"),
+		FPlatformMisc::GetUBTPlatform()
+	);
+
+	int32 ExitCode;
+
+	DesktopPlatform->RunUnrealBuildTool(
+		LOCTEXT("BuildingSwitchboardListener", "Building SwitchboardListener..."),
+		FPaths::RootDir(),
+		Arguments,
+		FeedbackContext,
+		ExitCode
+	);
+
+	return !ExitCode;
+}
 
 bool FSwitchboardEditorModule::LaunchListener()
 {
 	const FString ListenerPath = GetDefault<USwitchboardEditorSettings>()->GetListenerPlatformPath();
 	const FString ListenerArgs = GetDefault<USwitchboardEditorSettings>()->ListenerCommandlineArguments;
+
 	return RunProcess(ListenerPath, ListenerArgs);
 }
 
+bool FSwitchboardEditorModule::CreateNewConfig(const FSwitchboardNewConfigUserOptions& NewConfigUserOptions)
+{
+	FSwitchboardScriptArguments ScriptArgs;
+
+	ScriptArgs.ConfigName = NewConfigUserOptions.ConfigName;
+	ScriptArgs.EngineDir = FPaths::ConvertRelativePathToFull(FPaths::EngineDir());
+	ScriptArgs.ProjectPath = FPaths::ConvertRelativePathToFull(FPaths::GetProjectFilePath());
+
+	// Find persistent map
+	if (const UWorld* World = GEditor ? GEditor->GetEditorWorldContext(false).World() : nullptr)
+	{
+		FString ClassName, PackageName, ObjectName, SubObjectName;
+		FPackageName::SplitFullObjectPath(World->PersistentLevel.GetPath(), ClassName, PackageName, ObjectName, SubObjectName, true);
+		ScriptArgs.Map = PackageName;
+	}
+
+	// Get path to DCRA asset
+	if (const AActor* DCRA = NewConfigUserOptions.DCRA.DCRA.Get())
+	{
+		FString ClassName, PackageName, ObjectName, SubObjectName;
+
+		FPackageName::SplitFullObjectPath(DCRA->GetClass()->GetFullName(), ClassName, PackageName, ObjectName, SubObjectName, true);
+
+		const FString Extension = FPackageName::GetAssetPackageExtension();
+		const FString RelativeFilepath = FPackageName::LongPackageNameToFilename(PackageName, Extension);
+		const FString AbsoluteFilepath = FPaths::ConvertRelativePathToFull(RelativeFilepath);
+
+		if (FPaths::FileExists(AbsoluteFilepath))
+		{
+			ScriptArgs.DisplayClusterConfigPath = AbsoluteFilepath;
+		}
+	}
+
+	ScriptArgs.bUseLocalhost = NewConfigUserOptions.bUseLocalhost;
+	ScriptArgs.bAutoConnect = NewConfigUserOptions.bAutoConnect;
+	ScriptArgs.NumEditorDevices = NewConfigUserOptions.NumEditorDevices;
+
+	// Convert to Json and save it to a temp file
+
+	FString ScriptArgsString;
+	{
+		const bool bSBArgsOk = FJsonObjectConverter::UStructToJsonObjectString(ScriptArgs, ScriptArgsString, 0, 0, 0, nullptr, false);
+		check(bSBArgsOk); // This should always succeed because we're converting a known USTRUCT.
+	}
+
+	const FString ScriptArgsFilepath = FPaths::CreateTempFilename(
+		WCHAR_TO_TCHAR(std::filesystem::temp_directory_path().wstring().c_str()),
+		TEXT("sb_script_args_"),
+		TEXT(".json")
+	);
+
+	if (!FFileHelper::SaveStringToFile(ScriptArgsString, *ScriptArgsFilepath))
+	{
+		UE_LOG(LogSwitchboardPlugin, Log, TEXT("Could not save temp script arguments file '%s'"), *ScriptArgsFilepath);
+		return false;
+	}
+	
+	// The script that we will pass to Switchoard is in this plugin's Content/Python folder.
+	const FString PluginContentDir = IPluginManager::Get().FindPlugin(TEXT("Switchboard"))->GetContentDir();
+	const FString ScriptPath = UE::Switchboard::Private::ConcatPaths(PluginContentDir, "Python", "sb_script_new_config.py");
+
+	// Here we tell Switchboard to run the given script with the given argument, 
+	// and the argument is the path to a json file that will be interpreted by the script.
+	// We also pass --defaultenv so that SB doesn't interpret the first argument as the virtual environment.
+	const FString Arguments = FString::Printf(TEXT("--defaultenv --script \"%s\" --scriptargs \"%s\""),
+		*ScriptPath,
+		*ScriptArgsFilepath
+	);
+
+	UE_LOG(LogTemp, Warning, TEXT("%s"), *Arguments);
+
+	return LaunchSwitchboard(Arguments);
+}
 
 bool FSwitchboardEditorModule::RunProcess(const FString& InExe, const FString& InArgs)
 {
diff --git a/Engine/Plugins/VirtualProduction/Switchboard/Source/SwitchboardEditor/Private/SwitchboardEditorModule.h b/Engine/Plugins/VirtualProduction/Switchboard/Source/SwitchboardEditor/Private/SwitchboardEditorModule.h
--- a/Engine/Plugins/VirtualProduction/Switchboard/Source/SwitchboardEditor/Private/SwitchboardEditorModule.h
+++ b/Engine/Plugins/VirtualProduction/Switchboard/Source/SwitchboardEditor/Private/SwitchboardEditorModule.h
@@ -10,6 +10,7 @@
 #include "Modules/ModuleManager.h"
 #include "SwitchboardShortcuts.h"
 #include "SwitchboardScriptInterop.h"
+#include "SwitchboardTypes.h"
 
 
 
@@ -50,9 +51,25 @@ public:
 	virtual void ShutdownModule() override;
 	//~ End IModuleInterface
 
-	bool LaunchSwitchboard();
+	bool LaunchSwitchboard(const FString& Arguments = FString());
 	bool LaunchListener();
 
+	/**
+	 * Launches Switchboard with a new config that matches the current scene
+	 *
+	 * @param NewConfigUserOptions User options for New Config creation
+	 * 
+	 * @return Returns true if successful
+	 */
+	bool CreateNewConfig(const FSwitchboardNewConfigUserOptions& NewConfigUserOptions);
+
+	/**
+	 * Compiles SwitchboardListener (Development configuration)
+	 *
+	 * @return Returns true if successful
+	 */
+	bool CompileSwitchboardListener() const;
+
 	enum class ESwitchboardInstallState
 	{
 		VerifyInProgress,
@@ -92,15 +109,19 @@ public:
 	bool CreateOrUpdateShortcut(EShortcutApp App, EShortcutLocation Location);
 #endif
 
+	/** Returns the DCRA class, nullptr if it doesn't exist (i.e. nDisplay plugin is not enabled) */
+	static UClass* GetDisplayClusterRootActorClass();
+
 private:
 	void OnEngineInitComplete();
 	bool OnEditorSettingsModified();
 
 	void RunDefaultOSCListener();
 
-private:
 	bool RunProcess(const FString& InExe, const FString& InArgs);
 
+private:
+
 	FDelegateHandle DeferredStartDelegateHandle;
 
 	FString VerifyPath;
@@ -115,4 +136,7 @@ private:
 #if SWITCHBOARD_SHORTCUTS
 	TMap< TPair<EShortcutApp, EShortcutLocation>, EShortcutCompare > CachedShortcutCompares;
 #endif
+
+	/** Cached DCRA Class */
+	static TSoftObjectPtr<UClass> DisplayClusterRootActorClass;
 };
diff --git a/Engine/Plugins/VirtualProduction/Switchboard/Source/SwitchboardEditor/Private/SwitchboardMenuEntry.cpp b/Engine/Plugins/VirtualProduction/Switchboard/Source/SwitchboardEditor/Private/SwitchboardMenuEntry.cpp
--- a/Engine/Plugins/VirtualProduction/Switchboard/Source/SwitchboardEditor/Private/SwitchboardMenuEntry.cpp
+++ b/Engine/Plugins/VirtualProduction/Switchboard/Source/SwitchboardEditor/Private/SwitchboardMenuEntry.cpp
@@ -2,20 +2,25 @@
 
 #include "SwitchboardMenuEntry.h"
 
+#include "Dialog/SCustomDialog.h"
 #include "Framework/Application/SlateApplication.h"
 #include "Framework/Commands/Commands.h"
 #include "Framework/Commands/UIAction.h"
 #include "Framework/MultiBox/MultiBoxBuilder.h"
 #include "GenericPlatform/GenericPlatformMisc.h"
 #include "GenericPlatform/GenericPlatformProcess.h"
+#include "IStructureDetailsView.h"
+#include "Kismet/GameplayStatics.h"
 #include "LevelEditor.h"
 #include "Misc/FileHelper.h"
+#include "Misc/MessageDialog.h"
 #include "Misc/Paths.h"
 #include "Modules/ModuleManager.h"
 #include "SwitchboardEditorModule.h"
 #include "SwitchboardEditorSettings.h"
 #include "SwitchboardEditorStyle.h"
 #include "SwitchboardSetupWizard.h"
+#include "SwitchboardTypes.h"
 #include "ToolMenus.h"
 
 #define LOCTEXT_NAMESPACE "SwitchboardEditor"
@@ -32,10 +37,12 @@ public:
 	{
 		UI_COMMAND(LaunchSwitchboard, "Launch Switchboard", "Launch Switchboard", EUserInterfaceActionType::Button, FInputChord());
 		UI_COMMAND(LaunchSwitchboardListener, "Launch Switchboard Listener", "Launch Switchboard Listener", EUserInterfaceActionType::Button, FInputChord());
+		UI_COMMAND(CreateNewConfig, "Create Config", "Create Config", EUserInterfaceActionType::Button, FInputChord());
 	}
 
 	TSharedPtr<FUICommandInfo> LaunchSwitchboard;
 	TSharedPtr<FUICommandInfo> LaunchSwitchboardListener;
+	TSharedPtr<FUICommandInfo> CreateNewConfig;
 };
 
 
@@ -53,6 +60,10 @@ struct FSwitchboardMenuEntryImpl
 			FExecuteAction::CreateRaw(this, &FSwitchboardMenuEntryImpl::OnLaunchSwitchboardListener),
 			FCanExecuteAction());
 
+		Actions->MapAction(FSwitchboardUICommands::Get().CreateNewConfig,
+			FExecuteAction::CreateRaw(this, &FSwitchboardMenuEntryImpl::OnLaunchWithNewConfig),
+			FCanExecuteAction());
+
 		AddMenu();
 	}
 
@@ -66,21 +77,28 @@ struct FSwitchboardMenuEntryImpl
 
 		FToolMenuSection& Section = Menu->AddSection("Switchboard");
 
-		FToolMenuEntry SwitchboardButtonEntry = FToolMenuEntry::InitToolBarButton(FSwitchboardUICommands::Get().LaunchSwitchboard);
-		SwitchboardButtonEntry.SetCommandList(Actions);
-
-		const FToolMenuEntry SwitchboardComboEntry = FToolMenuEntry::InitComboButton(
-			"SwitchboardMenu",
-			FUIAction(),
-			FOnGetContent::CreateRaw(this, &FSwitchboardMenuEntryImpl::CreateListenerEntries),
-			TAttribute<FText>::Create([this]() { return LOCTEXT("LaunchSwitchboard", "Switchboard"); }),
-			LOCTEXT("SwitchboardTooltip", "Actions related to the SwitchboardListener"),
-			FSlateIcon(),
-			true //bInSimpleComboBox
-		);
-
-		Section.AddEntry(SwitchboardButtonEntry);
-		Section.AddEntry(SwitchboardComboEntry);
+		// Toolbar button
+		{
+			FToolMenuEntry SwitchboardButtonEntry = FToolMenuEntry::InitToolBarButton(FSwitchboardUICommands::Get().LaunchSwitchboard);
+			SwitchboardButtonEntry.SetCommandList(Actions);
+
+			Section.AddEntry(SwitchboardButtonEntry);
+		}
+
+		// Entries
+		{
+			const FToolMenuEntry SwitchboardComboEntry = FToolMenuEntry::InitComboButton(
+				"SwitchboardMenu",
+				FUIAction(),
+				FOnGetContent::CreateRaw(this, &FSwitchboardMenuEntryImpl::CreateMenuEntries),
+				TAttribute<FText>::Create([this]() { return LOCTEXT("LaunchSwitchboard", "Switchboard"); }),
+				LOCTEXT("SwitchboardTooltip", "Actions related to the SwitchboardListener"),
+				FSlateIcon(),
+				true //bInSimpleComboBox
+			);
+
+			Section.AddEntry(SwitchboardComboEntry);
+		}
 	}
 
 	void RemoveMenu()
@@ -100,7 +118,7 @@ struct FSwitchboardMenuEntryImpl
 		}
 	}
 
-	TSharedRef<SWidget> CreateListenerEntries()
+	TSharedRef<SWidget> CreateMenuEntries()
 	{
 		const bool bShouldCloseWindowAfterMenuSelection = true;
 		FMenuBuilder MenuBuilder(bShouldCloseWindowAfterMenuSelection, Actions);
@@ -131,16 +149,46 @@ struct FSwitchboardMenuEntryImpl
 		}
 		MenuBuilder.EndSection();
 
+		MenuBuilder.BeginSection("Switchboard", LOCTEXT("SwitchboardConfig", "Config"));
+		{
+			MenuBuilder.AddMenuEntry(FSwitchboardUICommands::Get().CreateNewConfig);
+		}
+		MenuBuilder.EndSection();
+
 		return MenuBuilder.MakeWidget();
 	}
 
 	void OnLaunchSwitchboardListener()
 	{
 		const FString ListenerPath = GetDefault<USwitchboardEditorSettings>()->GetListenerPlatformPath();
+
+		if (!FPaths::FileExists(ListenerPath))
+		{
+			// Since SwitchboardListener (SBL) doesn't exist, ask the user if he wants us to compile it.
+
+			const FText Msg = LOCTEXT("CouldNotFindSwitchboardListenerCompile", "Could not find SwitchboardListener. Would you like to compile it? ");
+
+			if (EAppReturnType::Yes == FMessageDialog::Open(EAppMsgType::YesNo, Msg))
+			{
+				// Compile SBL
+				if (!FSwitchboardEditorModule::Get().CompileSwitchboardListener())
+				{
+					const FText ErrorMsg = LOCTEXT("FailedToCompileSwitchboardListenerCompileIDE", "Failed to compile SwitchboardListener. Please compile with UGS or your IDE");
+					FMessageDialog::Open(EAppMsgType::Ok, ErrorMsg);
+					return;
+				}
+			}
+			else
+			{
+				// User did not want us to compile SBL, so we return.
+				return;
+			}
+		}
+
 		if (!FPaths::FileExists(ListenerPath))
 		{
-			const FString ErrorMsg = TEXT("Could not find SwitchboardListener! Make sure it was compiled.");
-			FPlatformMisc::MessageBoxExt(EAppMsgType::Ok, *ErrorMsg, TEXT("Error starting listener"));
+			const FText ErrorMsg = LOCTEXT("CouldNotFindSwitchboardListenerCompile", "Could not find SwitchboardListener! Make sure it was compiled.");
+			FMessageDialog::Open(EAppMsgType::Ok, ErrorMsg);
 			return;
 		}
 
@@ -150,8 +198,8 @@ struct FSwitchboardMenuEntryImpl
 		}
 		else
 		{
-			const FString ErrorMsg = TEXT("Unable to start the listener! Make sure it was compiled. Check the log for details.");
-			FPlatformMisc::MessageBoxExt(EAppMsgType::Ok, *ErrorMsg, TEXT("Error starting listener"));
+			const FText ErrorMsg = LOCTEXT("UnableToStartTheListener", "Unable to start the listener! Make sure it was compiled. Check the log for details.");
+			FMessageDialog::Open(EAppMsgType::Ok, ErrorMsg);
 		}
 	}
 
@@ -166,8 +214,8 @@ struct FSwitchboardMenuEntryImpl
 		{
 			if (!FPaths::FileExists(GetDefault<USwitchboardEditorSettings>()->GetListenerPlatformPath()))
 			{
-				const FString ErrorMsg = TEXT("Could not find SwitchboardListener! Make sure it has been compiled.");
-				FPlatformMisc::MessageBoxExt(EAppMsgType::Ok, *ErrorMsg, TEXT("Error enabling SwitchboardListener auto-launch"));
+				const FText ErrorMsg = LOCTEXT("CouldNotFindSwitchboardListenerCompile", "Could not find SwitchboardListener! Make sure it has been compiled.");
+				FMessageDialog::Open(EAppMsgType::Ok, ErrorMsg);
 				return;
 			}
 
@@ -191,11 +239,71 @@ struct FSwitchboardMenuEntryImpl
 		}
 		else
 		{
-			const FString ErrorMsg = TEXT("Unable to start Switchboard! Check the log for details.");
-			FPlatformMisc::MessageBoxExt(EAppMsgType::Ok, *ErrorMsg, TEXT("Error starting Switchboard"));
+			const FText ErrorMsg = LOCTEXT("UnableToStartSwitchboardCheckLog", "Unable to start Switchboard! Check the log for details.");
+			FMessageDialog::Open(EAppMsgType::Ok, ErrorMsg);
 		}
 	}
 
+	void OnLaunchWithNewConfig()
+	{
+		const FSwitchboardVerifyResult& VerifyResult = FSwitchboardEditorModule::Get().GetVerifyResult().Get();
+		if (VerifyResult.Summary != FSwitchboardVerifyResult::ESummary::Success)
+		{
+			SSwitchboardSetupWizard::OpenWindow();
+			return;
+		}
+
+		TSharedPtr<TStructOnScope<FSwitchboardNewConfigUserOptions>> NewConfigUserOptions
+			= MakeShared<TStructOnScope<FSwitchboardNewConfigUserOptions>>();
+
+		NewConfigUserOptions->InitializeAs<FSwitchboardNewConfigUserOptions>();
+
+		// Pre-populate the DCRA
+		if (UClass* DisplayClusterRootActorClass = FSwitchboardEditorModule::GetDisplayClusterRootActorClass())
+		{
+			const UWorld* World = GEditor ? GEditor->GetEditorWorldContext().World() : nullptr;
+
+			TArray<AActor*> DCRAs;
+			UGameplayStatics::GetAllActorsOfClass(World, DisplayClusterRootActorClass, DCRAs);
+
+			if (DCRAs.Num())
+			{
+				// Pick the first one
+				NewConfigUserOptions->Get()->DCRA.DCRA = DCRAs[0];
+			}
+		}
+
+		FPropertyEditorModule& PropertyEditor = FModuleManager::Get().LoadModuleChecked<FPropertyEditorModule>(TEXT("PropertyEditor"));
+
+		FStructureDetailsViewArgs StructureViewArgs;
+		FDetailsViewArgs DetailArgs;
+
+		DetailArgs.bAllowSearch = false;
+		DetailArgs.bShowScrollBar = true;
+
+		TSharedPtr<IStructureDetailsView> NewConfigUsernOptionsDetailsView
+			= PropertyEditor.CreateStructureDetailView(DetailArgs, StructureViewArgs, NewConfigUserOptions);
+
+		TSharedRef<SCustomDialog> OptionsWindow = SNew(SCustomDialog)
+			.Title(LOCTEXT("SwitchboardNewConfigOptions", "Switchboard New Config Options"))
+			.ToolTipText(LOCTEXT("CreateNewConfigTooltip", "This will open Switchboard and create a new config from the parameters below and the current map."))
+			.Content()
+			[
+				NewConfigUsernOptionsDetailsView->GetWidget().ToSharedRef()
+			]
+			.Buttons
+			({
+				SCustomDialog::FButton(LOCTEXT("Ok", "Ok"), FSimpleDelegate::CreateLambda([NewConfigUserOptions]() -> void
+				{
+					check(NewConfigUserOptions.IsValid());
+					FSwitchboardEditorModule::Get().CreateNewConfig(*NewConfigUserOptions->Get());
+				})),
+				SCustomDialog::FButton(LOCTEXT("Cancel", "Cancel")),
+			});
+
+		OptionsWindow->Show(); // Showing it non-modal mostly because its DCRA actor picker doesn't work if the window is modal.
+	}
+
 public:
 	static TUniquePtr<FSwitchboardMenuEntryImpl> Implementation;
 
diff --git a/Engine/Plugins/VirtualProduction/Switchboard/Source/SwitchboardEditor/Private/SwitchboardTypes.h b/Engine/Plugins/VirtualProduction/Switchboard/Source/SwitchboardEditor/Private/SwitchboardTypes.h
new file mode 100644
--- /dev/null
+++ b/Engine/Plugins/VirtualProduction/Switchboard/Source/SwitchboardEditor/Private/SwitchboardTypes.h
@@ -0,0 +1,96 @@
+// Copyright Epic Games, Inc. All Rights Reserved.
+
+#pragma once
+
+#include "CoreMinimal.h"
+
+#include "SwitchboardTypes.generated.h"
+
+/**
+ * Structure that holds the parameters to the sb_script_new_config.py script that will be passed to Switchboard when auto-creating a new config.
+ * This structure is saved as a temporary json file and its path is passed to SB in the command line.
+ */
+USTRUCT()
+struct FSwitchboardScriptArguments
+{
+	GENERATED_BODY()
+
+	/* Optional name of the new config. If left empty, a name will be autogenerated. */
+	UPROPERTY()
+	FString ConfigName;
+
+	/* Path to the Engine directory so that Switchboard knows which engine to launch */
+	UPROPERTY()
+	FString EngineDir;
+
+	/* Path to the Project directory so that Switchboard knows which project to launch */
+	UPROPERTY()
+	FString ProjectPath;
+
+	/* Gamepath to map that Switchboard should launch */
+	UPROPERTY()
+	FString Map;
+
+	/* Path to nDisplay config that Switchboard should create the nDisplay cluster from */
+	UPROPERTY()
+	FString DisplayClusterConfigPath;
+
+	/* Number of Editor devices that it should create */
+	UPROPERTY()
+	int32 NumEditorDevices = 1;
+
+	/* Indicates that Switchboard should override IP addresses with localhost */
+	UPROPERTY()
+	bool bUseLocalhost = true;
+
+	/* Indicates that Switchboard should auto-connect to the listeners of each created device */
+	UPROPERTY()
+	bool bAutoConnect = true;
+
+	/* Indicates that Switchboard should delete the script arguments file */
+	UPROPERTY()
+	bool bAutoDeleteScriptArgsFile = true;
+};
+
+/**
+ * Structure used to hold a reference of an nDisplay actor for detail customization purposes
+ * Note: The reason for not simply using <ADisplayClusterRootActor> is to avoid depending on the nDisplay plugin
+ */
+USTRUCT(BlueprintType)
+struct FDisplayClusterRootActorReference
+{
+	GENERATED_BODY()
+
+	/** Points to the nDisplay actor */
+	UPROPERTY()
+	TSoftObjectPtr<AActor> DCRA;
+};
+
+/**
+ * Structure used to collect New Config options from the user
+ */
+USTRUCT(BlueprintType)
+struct FSwitchboardNewConfigUserOptions
+{
+	GENERATED_BODY()
+
+	/* Optional name of the new config. If left empty, a name will be autogenerated. */
+	UPROPERTY(EditAnywhere, BlueprintReadWrite, Category = "General")
+	FString ConfigName;
+
+	/* Indicates which nDisplay actor to use */
+	UPROPERTY(EditAnywhere, BlueprintReadWrite, Category = "Devices", meta = (DisplayName = "nDisplay Actor"))
+	FDisplayClusterRootActorReference DCRA;
+
+	/* Number of Editor devices that it should create */
+	UPROPERTY(EditAnywhere, BlueprintReadWrite, Category = "Devices", meta = (DisplayName = "Editors"))
+	int32 NumEditorDevices = 1;
+
+	/* Indicates that Switchboard should override IP addresses with localhost */
+	UPROPERTY(EditAnywhere, BlueprintReadWrite, Category = "Network")
+	bool bUseLocalhost = true;
+
+	/* Indicates that Switchboard should auto-connect to the listeners of each created device */
+	UPROPERTY(EditAnywhere, BlueprintReadWrite, Category = "Network")
+	bool bAutoConnect = true;
+};
diff --git a/Engine/Plugins/VirtualProduction/Switchboard/Source/SwitchboardEditor/SwitchboardEditor.build.cs b/Engine/Plugins/VirtualProduction/Switchboard/Source/SwitchboardEditor/SwitchboardEditor.build.cs
--- a/Engine/Plugins/VirtualProduction/Switchboard/Source/SwitchboardEditor/SwitchboardEditor.build.cs
+++ b/Engine/Plugins/VirtualProduction/Switchboard/Source/SwitchboardEditor/SwitchboardEditor.build.cs
@@ -8,15 +8,19 @@ public class SwitchboardEditor : ModuleRules
 	{
 		PrivateDependencyModuleNames.AddRange(
 			new string[] {
+				"AssetRegistry",
 				"Blutility",
 				"Core",
 				"CoreUObject",
+				"DesktopPlatform",
 				"Engine",
 				"EditorStyle",
 				"InputCore",
 				"Json",
+				"JsonUtilities",
 				"MessageLog",
 				"Projects",
+				"PropertyEditor",
 				"Settings",
 				"Slate",
 				"SlateCore",
