Added setting to launch multi-user server in slate mode

diff --git a/Engine/Plugins/VirtualProduction/Switchboard/Source/Switchboard/switchboard/config.py b/Engine/Plugins/VirtualProduction/Switchboard/Source/Switchboard/switchboard/config.py
--- a/Engine/Plugins/VirtualProduction/Switchboard/Source/Switchboard/switchboard/config.py
+++ b/Engine/Plugins/VirtualProduction/Switchboard/Source/Switchboard/switchboard/config.py
@@ -2187,6 +2187,11 @@ class Config(object):
                 "Auto Launch",
                 data.get('muserver_auto_launch', True)
             ),
+            "muserver_slate_mode": BoolSetting(
+                "muserver_slate_mode",
+                "Launch Multi-user server in UI mode",
+                data.get('muserver_slate_mode', True)
+            ),
             "muserver_clean_history": BoolSetting(
                 "muserver_clean_history",
                 "Clean History",
@@ -2215,6 +2220,7 @@ class Config(object):
         self.MUSERVER_MULTICAST_ENDPOINT = self.mu_settings["udpmessaging_multicast_endpoint"]
         self.MULTIUSER_SERVER_EXE = self.mu_settings["multiuser_exe"]
         self.MUSERVER_AUTO_LAUNCH = self.mu_settings["muserver_auto_launch"]
+        self.MUSERVER_SLATE_MODE = self.mu_settings["muserver_slate_mode"]
         self.MUSERVER_CLEAN_HISTORY = self.mu_settings["muserver_clean_history"]
         self.MUSERVER_AUTO_BUILD = self.mu_settings["muserver_auto_build"]
         self.MUSERVER_AUTO_ENDPOINT = self.mu_settings["muserver_auto_endpoint"]
@@ -2231,6 +2237,7 @@ class Config(object):
         data["muserver_endpoint"] = self.MUSERVER_ENDPOINT.get_value()
         data["multiuser_exe"] = self.MULTIUSER_SERVER_EXE.get_value()
         data["muserver_auto_launch"] = self.MUSERVER_AUTO_LAUNCH.get_value()
+        data["muserver_slate_mode"] = self.MUSERVER_SLATE_MODE.get_value()
         data["muserver_clean_history"] = self.MUSERVER_CLEAN_HISTORY.get_value()
         data["muserver_auto_build"] = self.MUSERVER_AUTO_BUILD.get_value()
         data["muserver_auto_endpoint"] = self.MUSERVER_AUTO_ENDPOINT.get_value()
diff --git a/Engine/Plugins/VirtualProduction/Switchboard/Source/Switchboard/switchboard/switchboard_application.py b/Engine/Plugins/VirtualProduction/Switchboard/Source/Switchboard/switchboard/switchboard_application.py
--- a/Engine/Plugins/VirtualProduction/Switchboard/Source/Switchboard/switchboard/switchboard_application.py
+++ b/Engine/Plugins/VirtualProduction/Switchboard/Source/Switchboard/switchboard/switchboard_application.py
@@ -133,7 +133,10 @@ class MultiUserApplication:
 
             cmdline = ''
             if sys.platform.startswith('win'):
-                cmdline = f'start "Multi User Server" "{self.exe_path()}"'
+                if CONFIG.MUSERVER_SLATE_MODE.get_value():
+                    cmdline = f'start /b "Multi User Server" "{self.exe_path()}"'
+                else:
+                    cmdline = f'start "Multi User Server" "{self.exe_path()}"'
             else:
                 cmdline = f'{self.exe_path()}'
 
@@ -147,6 +150,9 @@ class MultiUserApplication:
 
             if CONFIG.MUSERVER_CLEAN_HISTORY.get_value():
                 cmdline += " -ConcertClean"
+                
+            if CONFIG.MUSERVER_SLATE_MODE.get_value():
+                cmdline += " -WithSlate"
 
             if len(args) > 0:
                 cmdline += f' {" ".join(args)}'
