Switchboard: Removed leftover log.

diff --git a/Engine/Plugins/VirtualProduction/Switchboard/Source/SwitchboardEditor/Private/SwitchboardEditorModule.cpp b/Engine/Plugins/VirtualProduction/Switchboard/Source/SwitchboardEditor/Private/SwitchboardEditorModule.cpp
--- a/Engine/Plugins/VirtualProduction/Switchboard/Source/SwitchboardEditor/Private/SwitchboardEditorModule.cpp
+++ b/Engine/Plugins/VirtualProduction/Switchboard/Source/SwitchboardEditor/Private/SwitchboardEditorModule.cpp
@@ -281,8 +281,6 @@ bool FSwitchboardEditorModule::CreateNewConfig(const FSwitchboardNewConfigUserOp
 		*ScriptArgsFilepath
 	);
 
-	UE_LOG(LogTemp, Warning, TEXT("%s"), *Arguments);
-
 	return LaunchSwitchboard(Arguments);
 }
 
