Switchboard: Fix LocDuplicates

diff --git a/Engine/Plugins/VirtualProduction/Switchboard/Source/SwitchboardEditor/Private/SwitchboardMenuEntry.cpp b/Engine/Plugins/VirtualProduction/Switchboard/Source/SwitchboardEditor/Private/SwitchboardMenuEntry.cpp
--- a/Engine/Plugins/VirtualProduction/Switchboard/Source/SwitchboardEditor/Private/SwitchboardMenuEntry.cpp
+++ b/Engine/Plugins/VirtualProduction/Switchboard/Source/SwitchboardEditor/Private/SwitchboardMenuEntry.cpp
@@ -166,7 +166,7 @@ struct FSwitchboardMenuEntryImpl
 		{
 			// Since SwitchboardListener (SBL) doesn't exist, ask the user if he wants us to compile it.
 
-			const FText Msg = LOCTEXT("CouldNotFindSwitchboardListenerCompile", "Could not find SwitchboardListener. Would you like to compile it? ");
+			const FText Msg = LOCTEXT("CouldNotFindSwitchboardListenerCompileAskIfCompile", "Could not find SwitchboardListener. Would you like to compile it? ");
 
 			if (EAppReturnType::Yes == FMessageDialog::Open(EAppMsgType::YesNo, Msg))
 			{
@@ -187,7 +187,7 @@ struct FSwitchboardMenuEntryImpl
 
 		if (!FPaths::FileExists(ListenerPath))
 		{
-			const FText ErrorMsg = LOCTEXT("CouldNotFindSwitchboardListenerCompile", "Could not find SwitchboardListener! Make sure it was compiled.");
+			const FText ErrorMsg = LOCTEXT("CouldNotFindSwitchboardListenerCompileMakeSureCompiled", "Could not find SwitchboardListener! Make sure it was compiled.");
 			FMessageDialog::Open(EAppMsgType::Ok, ErrorMsg);
 			return;
 		}
@@ -214,7 +214,7 @@ struct FSwitchboardMenuEntryImpl
 		{
 			if (!FPaths::FileExists(GetDefault<USwitchboardEditorSettings>()->GetListenerPlatformPath()))
 			{
-				const FText ErrorMsg = LOCTEXT("CouldNotFindSwitchboardListenerCompile", "Could not find SwitchboardListener! Make sure it has been compiled.");
+				const FText ErrorMsg = LOCTEXT("CouldNotFindSwitchboardListenerCompileMakeSureHasBeenCompiled", "Could not find SwitchboardListener! Make sure it has been compiled.");
 				FMessageDialog::Open(EAppMsgType::Ok, ErrorMsg);
 				return;
 			}
