Switchboard: Make Save Config As do a copy instead of  replace of the config file.

Project name is also now updated accordingly.

diff --git a/Engine/Plugins/VirtualProduction/Switchboard/Source/Switchboard/switchboard/config.py b/Engine/Plugins/VirtualProduction/Switchboard/Source/Switchboard/switchboard/config.py
--- a/Engine/Plugins/VirtualProduction/Switchboard/Source/Switchboard/switchboard/config.py
+++ b/Engine/Plugins/VirtualProduction/Switchboard/Source/Switchboard/switchboard/config.py
@@ -2297,6 +2297,30 @@ class Config(object):
         self.file_path = new_config_path
         self.save()
 
+    def save_as(self, new_config_path: typing.Union[str, pathlib.Path]):
+        """
+        Copy the file.
+
+        If a file already exists at the new path, it will be overwritten.
+        """
+        new_config_path = get_absolute_config_path(new_config_path)
+
+        new_project_name = new_config_path.stem
+
+        if self.file_path:
+            new_config_path.parent.mkdir(parents=True, exist_ok=True)
+            shutil.copyfile(src=self.file_path, dst=new_config_path)
+
+            # Don't change project name if the user customized it
+            if self.PROJECT_NAME.get_value() != self.file_path.stem:
+                new_project_name = self.file_path.stem
+
+        self.file_path = new_config_path
+
+        self.PROJECT_NAME.update_value(new_project_name)
+
+        self.save()
+
     def save(self):
         if not self.file_path or not self.saving_allowed:
             return
diff --git a/Engine/Plugins/VirtualProduction/Switchboard/Source/Switchboard/switchboard/switchboard_dialog.py b/Engine/Plugins/VirtualProduction/Switchboard/Source/Switchboard/switchboard/switchboard_dialog.py
--- a/Engine/Plugins/VirtualProduction/Switchboard/Source/Switchboard/switchboard/switchboard_dialog.py
+++ b/Engine/Plugins/VirtualProduction/Switchboard/Source/Switchboard/switchboard/switchboard_dialog.py
@@ -904,7 +904,7 @@ class SwitchboardDialog(QtCore.QObject):
             return
 
         # Switch to the new config
-        CONFIG.replace(new_config_path)
+        CONFIG.save_as(new_config_path)
 
         # Update the settings
         SETTINGS.CONFIG = new_config_path
