Add working and saved directory configuation options to multi-user server.

diff --git a/Engine/Plugins/VirtualProduction/Switchboard/Source/Switchboard/switchboard/config.py b/Engine/Plugins/VirtualProduction/Switchboard/Source/Switchboard/switchboard/config.py
--- a/Engine/Plugins/VirtualProduction/Switchboard/Source/Switchboard/switchboard/config.py
+++ b/Engine/Plugins/VirtualProduction/Switchboard/Source/Switchboard/switchboard/config.py
@@ -2187,6 +2187,16 @@ class Config(object):
                 "Multiuser Slate Executable Name",
                 data.get('multiuserslate_exe', 'UnrealMultiUserSlateServer')
             ),
+            "muserver_archive_dir": DirectoryPathSetting(
+                "muserver_archive_dir",
+                "Directory for Saved Archives",
+                data.get('muserver_archive_dir', '')
+            ),
+            "muserver_working_dir": DirectoryPathSetting(
+                "muserver_working_dir",
+                "Directory for Live Sessions",
+                data.get('muserver_working_dir', '')
+            ),
             "muserver_auto_launch": BoolSetting(
                 "muserver_auto_launch",
                 "Auto Launch",
@@ -2218,7 +2228,7 @@ class Config(object):
                 data.get('muserver_auto_join', True)
             )
         }
-        
+
         self.MUSERVER_SERVER_NAME = self.mu_settings["muserver_server_name"]
         self.MUSERVER_COMMAND_LINE_ARGUMENTS = self.mu_settings["muserver_command_line_arguments"]
         self.MUSERVER_ENDPOINT = self.mu_settings["muserver_endpoint"]
@@ -2231,12 +2241,14 @@ class Config(object):
         self.MUSERVER_AUTO_BUILD = self.mu_settings["muserver_auto_build"]
         self.MUSERVER_AUTO_ENDPOINT = self.mu_settings["muserver_auto_endpoint"]
         self.MUSERVER_AUTO_JOIN = self.mu_settings["muserver_auto_join"]
+        self.MUSERVER_WORKING_DIR = self.mu_settings["muserver_working_dir"]
+        self.MUSERVER_ARCHIVE_DIR = self.mu_settings["muserver_archive_dir"]
 
     def save_unreal_insights(self, data):
         data['tracing_enabled'] = self.INSIGHTS_TRACE_ENABLE.get_value()
         data['tracing_args'] = self.INSIGHTS_TRACE_ARGS.get_value()
         data['tracing_stat_events'] = self.INSIGHTS_STAT_EVENTS.get_value()
-        
+
     def save_muserver(self, data):
         data["muserver_command_line_arguments"] = self.MUSERVER_COMMAND_LINE_ARGUMENTS.get_value()
         data["muserver_server_name"] = self.MUSERVER_SERVER_NAME.get_value()
@@ -2250,6 +2262,8 @@ class Config(object):
         data["muserver_auto_endpoint"] = self.MUSERVER_AUTO_ENDPOINT.get_value()
         data["muserver_multicast_endpoint"] = self.MUSERVER_MULTICAST_ENDPOINT.get_value()
         data["muserver_auto_join"] = self.MUSERVER_AUTO_JOIN.get_value()
+        data["muserver_archive_dir"] = self.MUSERVER_ARCHIVE_DIR.get_value()
+        data["muserver_working_dir"] = self.MUSERVER_WORKING_DIR.get_value()
 
     def load_plugin_settings(self, device_type, settings):
         ''' Updates plugin settings values with those read from the config file.
@@ -2476,6 +2490,9 @@ class Config(object):
                 self.ENGINE_DIR.get_value(), self.MULTIUSER_SERVER_EXE.get_value())
 
     def multiuser_server_session_directory_path(self):
+        if self.MUSERVER_WORKING_DIR.get_value():
+            return self.MUSERVER_WORKING_DIR.get_value()
+
         if self.MUSERVER_SLATE_MODE.get_value():
             return os.path.join(self.ENGINE_DIR.get_value(), "Programs", "UnrealMultiUserSlateServer", "Intermediate", "MultiUser")
 
diff --git a/Engine/Plugins/VirtualProduction/Switchboard/Source/Switchboard/switchboard/devices/unreal/plugin_unreal.py b/Engine/Plugins/VirtualProduction/Switchboard/Source/Switchboard/switchboard/devices/unreal/plugin_unreal.py
--- a/Engine/Plugins/VirtualProduction/Switchboard/Source/Switchboard/switchboard/devices/unreal/plugin_unreal.py
+++ b/Engine/Plugins/VirtualProduction/Switchboard/Source/Switchboard/switchboard/devices/unreal/plugin_unreal.py
@@ -1131,6 +1131,9 @@ class DeviceUnreal(Device):
 
                 puuid_dependency = self._build_mu_server(
                     puuid_dependency=puuid_dependency)
+ 
+                puuid_dependency = self._build_mu_slate_server(
+                    puuid_dependency=puuid_dependency)
 
             # TODO: Build listener
             # - Handle the running listener's exe being locked
@@ -1165,6 +1168,12 @@ class DeviceUnreal(Device):
         return self._queue_build(
             'mu_server', ubt_args=ubt_args, puuid_dependency=puuid_dependency)
 
+    def _build_mu_slate_server(self, puuid_dependency: Optional[uuid.UUID] = None):
+        ubt_args = (f'UnrealMultiUserSlateServer {self.target_platform} '
+                    'Development -Progress')
+        return self._queue_build(
+            'mu_slate_server', ubt_args=ubt_args, puuid_dependency=puuid_dependency)
+    
     def _build_listener(self, puuid_dependency: Optional[uuid.UUID] = None):
         ubt_args = (f'SwitchboardListener {self.target_platform} '
                     'Development -Progress')
diff --git a/Engine/Plugins/VirtualProduction/Switchboard/Source/Switchboard/switchboard/switchboard_application.py b/Engine/Plugins/VirtualProduction/Switchboard/Source/Switchboard/switchboard/switchboard_application.py
--- a/Engine/Plugins/VirtualProduction/Switchboard/Source/Switchboard/switchboard/switchboard_application.py
+++ b/Engine/Plugins/VirtualProduction/Switchboard/Source/Switchboard/switchboard/switchboard_application.py
@@ -148,6 +148,12 @@ class MultiUserApplication:
             if self.concert_ignore_cl:
                 cmdline += " -ConcertIgnore"
 
+            if CONFIG.MUSERVER_WORKING_DIR.get_value():
+                cmdline += f' -ConcertWorkingDir="{CONFIG.MUSERVER_WORKING_DIR.get_value()}"'
+
+            if CONFIG.MUSERVER_ARCHIVE_DIR.get_value():
+                cmdline += f' -ConcertSavedDir="{CONFIG.MUSERVER_ARCHIVE_DIR.get_value()}"'
+
             if CONFIG.MUSERVER_CLEAN_HISTORY.get_value():
                 cmdline += " -ConcertClean"
 
