Switchboard: FIxed issue where opening the settings dialog expected a relative path to the config file to exist (which is not the case whe it is in a different drive letter).

#skipci

diff --git a/Engine/Plugins/VirtualProduction/Switchboard/Source/Switchboard/switchboard/settings_dialog.py b/Engine/Plugins/VirtualProduction/Switchboard/Source/Switchboard/switchboard/settings_dialog.py
--- a/Engine/Plugins/VirtualProduction/Switchboard/Source/Switchboard/switchboard/settings_dialog.py
+++ b/Engine/Plugins/VirtualProduction/Switchboard/Source/Switchboard/switchboard/settings_dialog.py
@@ -322,12 +322,15 @@ class SettingsDialog(QtCore.QObject):
 
         if not config_path:
             self._changed_config_path = config.Config.DEFAULT_CONFIG_PATH
-            rel_config_path_str = self._changed_config_path.stem
+            config_path_str = self._changed_config_path.stem
         else:
-            rel_config_path_str = str(
-                config.get_relative_config_path(config_path).with_suffix(''))
+            # prefer relative path when possible
+            try:
+                config_path_str = str(config.get_relative_config_path(config_path).with_suffix(''))
+            except ValueError:
+                config_path_str = str(config_path)
 
-        self.ui.config_path_line_edit.setText(rel_config_path_str)
+        self.ui.config_path_line_edit.setText(config_path_str)
 
     def config_path_text_changed(self, config_path_str):
         config_path = config.Config.DEFAULT_CONFIG_PATH
