Swtichboard: Support new ObjectClassName format in .uasset files.

#skipci

diff --git a/Engine/Plugins/VirtualProduction/Switchboard/Source/Switchboard/switchboard/devices/ndisplay/plugin_ndisplay.py b/Engine/Plugins/VirtualProduction/Switchboard/Source/Switchboard/switchboard/devices/ndisplay/plugin_ndisplay.py
--- a/Engine/Plugins/VirtualProduction/Switchboard/Source/Switchboard/switchboard/devices/ndisplay/plugin_ndisplay.py
+++ b/Engine/Plugins/VirtualProduction/Switchboard/Source/Switchboard/switchboard/devices/ndisplay/plugin_ndisplay.py
@@ -205,8 +205,9 @@ class AddnDisplayDialog(AddDeviceDialog):
             ''' Returns the asset if it is an nDisplay config '''
             with open(asset['path'], 'rb') as file:
                 aparser = UassetParser(file, allowUnversioned=True)
+                validConfigClassNames = ('DisplayClusterBlueprint','/Script/DisplayCluster.DisplayClusterBlueprint')
                 for assetdata in aparser.aregdata:
-                    if assetdata.ObjectClassName == 'DisplayClusterBlueprint':
+                    if assetdata.ObjectClassName in validConfigClassNames:
                         return asset
             raise ValueError
 
