Swtichboard: Support new ObjectClassName format in .uasset files.

Added missing support when device was being added.

#skipci

diff --git a/Engine/Plugins/VirtualProduction/Switchboard/Source/Switchboard/switchboard/devices/ndisplay/plugin_ndisplay.py b/Engine/Plugins/VirtualProduction/Switchboard/Source/Switchboard/switchboard/devices/ndisplay/plugin_ndisplay.py
--- a/Engine/Plugins/VirtualProduction/Switchboard/Source/Switchboard/switchboard/devices/ndisplay/plugin_ndisplay.py
+++ b/Engine/Plugins/VirtualProduction/Switchboard/Source/Switchboard/switchboard/devices/ndisplay/plugin_ndisplay.py
@@ -201,14 +201,17 @@ class AddnDisplayDialog(AddDeviceDialog):
         # Looks much better without the window frame.
         progressDiag.setWindowFlag(QtCore.Qt.FramelessWindowHint)
 
+        ''' Returns the asset if it is an nDisplay config '''
         def validateConfigAsset(asset):
-            ''' Returns the asset if it is an nDisplay config '''
+
             with open(asset['path'], 'rb') as file:
+
                 aparser = UassetParser(file, allowUnversioned=True)
-                validConfigClassNames = ('DisplayClusterBlueprint','/Script/DisplayCluster.DisplayClusterBlueprint')
+
                 for assetdata in aparser.aregdata:
-                    if assetdata.ObjectClassName in validConfigClassNames:
+                    if assetdata.ObjectClassName in DisplayConfig.validClassNames:
                         return asset
+
             raise ValueError
 
         numThreads = 8
@@ -345,6 +348,8 @@ class DeviceWidgetnDisplay(DeviceWidgetUnreal):
 class DisplayConfig(object):
     ''' Encapsulates nDisplay config'''
 
+    validClassNames = ('DisplayClusterBlueprint','/Script/DisplayCluster.DisplayClusterBlueprint')
+
     def __init__(self):
         self.nodes = []
         self.uasset_path = ''
@@ -1020,8 +1025,9 @@ class DevicenDisplay(DeviceUnreal):
 
             # Check that the asset is of the right class, then find its
             # ConfigExport tag and parse it.
+
             for assetdata in aparser.aregdata:
-                if assetdata.ObjectClassName == 'DisplayClusterBlueprint':
+                if assetdata.ObjectClassName in DisplayConfig.validClassNames:
                     return assetdata.tags['ConfigExport']
 
         raise ValueError('Invalid nDisplay config .uasset')
