Make sure we update the state of the header widget value with the values of the children.

diff --git a/Engine/Plugins/VirtualProduction/Switchboard/Source/Switchboard/switchboard/device_list_widget.py b/Engine/Plugins/VirtualProduction/Switchboard/Source/Switchboard/switchboard/device_list_widget.py
--- a/Engine/Plugins/VirtualProduction/Switchboard/Source/Switchboard/switchboard/device_list_widget.py
+++ b/Engine/Plugins/VirtualProduction/Switchboard/Source/Switchboard/switchboard/device_list_widget.py
@@ -24,8 +24,11 @@ class DeviceListWidget(QtWidgets.QListWidget):
         #self.setSelectionMode(QtWidgets.QAbstractItemView.NoSelection)
         self.setSelectionMode(QtWidgets.QAbstractItemView.MultiSelection)
 
-        self._device_widgets: Dict[int, DeviceWidget] = {}  # key = device_hash
+        self._devices: Dict[int, Device] = {}  # key = device_hash
+
         self._header_by_category_name = {}
+
+        self._autojoin_device_guard = False
         self.setMinimumSize(self.minimumSize().width(), 5)
 
     # Override all mouse events so that selection can be used for qss styling
@@ -50,7 +53,7 @@ class DeviceListWidget(QtWidgets.QListWidget):
 
         item_widget = self.itemWidget(item)
         assert item_widget
-        
+
         # No actions for clicking the header widget
         if isinstance(item_widget, DeviceWidgetHeader):
             return
@@ -72,6 +75,9 @@ class DeviceListWidget(QtWidgets.QListWidget):
             self.signal_remove_device.emit(item_widget.device_hash)
 
     def devices_in_category(self, category):
+        if category not in self._header_by_category_name:
+            return []
+
         header_item = self._header_by_category_name[category]
         header_row = self.row(header_item)
         device_widgets = []
@@ -83,7 +89,7 @@ class DeviceListWidget(QtWidgets.QListWidget):
         return device_widgets
 
     def on_device_removed(self, device_hash, device_type, *args):
-        self._device_widgets.pop(device_hash)
+        self._devices.pop(device_hash)
 
         for i in range(self.count()):
             item = self.item(i)
@@ -99,6 +105,8 @@ class DeviceListWidget(QtWidgets.QListWidget):
             self.takeItem(self.row(header_item))
             self._header_by_category_name.pop(category)
 
+        self.update_header_autojoin_state(category)
+
     def on_connect_all_toggled(self, plugin_name, state):
         self.signal_connect_all_plugin_devices_toggled.emit(plugin_name, state)
 
@@ -119,6 +127,34 @@ class DeviceListWidget(QtWidgets.QListWidget):
         self.attach_autojoin_button(header_widget, label_name)
         return device_item
 
+    def update_header_autojoin_state(self, category):
+        if self._autojoin_device_guard:
+            return
+
+        header_item = self._header_by_category_name.get(category)
+        if not header_item:
+            return
+
+        header_widget = self.itemWidget(header_item)
+        assert isinstance(header_widget, DeviceWidgetHeader)
+
+        any_set = False
+        for device in self.devices():
+            is_autojoin_set = device.autojoin_mu_server and device.autojoin_mu_server.get_value()
+            any_set = any_set or is_autojoin_set
+
+        wasBlockingSignals = header_widget.autojoin_mu.get_button().blockSignals(True)
+        header_widget.autojoin_mu.get_button().setChecked(any_set)
+        header_widget.autojoin_mu.get_button().blockSignals(wasBlockingSignals)
+
+    def connect_to_device_autojoin(self, device, category):
+        def child_device_changed():
+            self.update_header_autojoin_state(category)
+
+        if device.autojoin_mu_server:
+            device.autojoin_mu_server.signal_setting_changed.connect(child_device_changed)
+            self.update_header_autojoin_state(category)
+
     def attach_autojoin_button(self, header_widget, category):
         if header_widget.autojoin_mu is None:
             return
@@ -126,13 +162,18 @@ class DeviceListWidget(QtWidgets.QListWidget):
 
         def do_set_autojoin():
             self.set_autojoin_mu(autojoin, category)
+
         autojoin.signal_device_widget_autojoin_mu.connect(do_set_autojoin)
 
     def set_autojoin_mu(self, autojoin_ui, category):
+        self._autojoin_device_guard = True
+
         is_checked = autojoin_ui.is_autojoin_enabled()
         for device in self.devices_in_category(category):
             device.autojoin_mu.set_autojoin_mu(is_checked)
 
+        self._autojoin_device_guard = False
+
     def add_device_widget(self, device: Device):
         category = device.category_name
         if category not in self._header_by_category_name.keys():
@@ -153,17 +194,19 @@ class DeviceListWidget(QtWidgets.QListWidget):
         self.setItemWidget(device_item, device.widget)
 
         # Keep a dict for quick lookup
-        self._device_widgets[device.device_hash] = device.widget
+        self._devices[device.device_hash] = device
 
         self.signal_register_device_widget.emit(device.widget)
         # force the widget to update, to make sure status is correct
         device.widget.update_status(device.status, device.status)
 
+        self.connect_to_device_autojoin(device, category)
+
         return device.widget
 
     def clear_widgets(self):
         self.clear()
-        self._device_widgets = {}
+        self._devices = {}
         self._header_by_category_name = {}
 
     def widget_item_by_device(self, device: Device):
@@ -179,13 +222,17 @@ class DeviceListWidget(QtWidgets.QListWidget):
         return None
 
     def device_widget_by_hash(self, device_hash: int):
-        if device_hash not in self._device_widgets:
+        if device_hash not in self._devices:
             return None
 
-        return self._device_widgets[device_hash]
+        return self._devices[device_hash].widget
+
+    def devices(self):
+        return self._devices.values()
 
     def device_widgets(self):
-        return self._device_widgets.values()
+        widgets = [device.widget for device in self._devices.values()]
+        return widgets
 
     def update_category_status(
             self,
