Return true / false on collect logs so that we can determine if we show the file explorer.

diff --git a/Engine/Plugins/VirtualProduction/Switchboard/Source/Switchboard/switchboard/switchboard_dialog.py b/Engine/Plugins/VirtualProduction/Switchboard/Source/Switchboard/switchboard/switchboard_dialog.py
--- a/Engine/Plugins/VirtualProduction/Switchboard/Source/Switchboard/switchboard/switchboard_dialog.py
+++ b/Engine/Plugins/VirtualProduction/Switchboard/Source/Switchboard/switchboard/switchboard_dialog.py
@@ -562,14 +562,15 @@ class SwitchboardDialog(QtCore.QObject):
         ''' Registers convenience "Open Logs Folder" menu item '''
         action = self.register_tools_menu_action("&Open Logs Folder")
         action.triggered.connect(collect_logs.open_logs_folder)
-     
+
     def register_zip_logs_menuitem(self):
         def save_logs():
-            collect_logs.execute_zip_logs_workflow(
+            status = collect_logs.execute_zip_logs_workflow(
                 CONFIG,
                 [device for device in self.device_manager.devices() if isinstance(device, DeviceUnreal)]
             )
-            collect_logs.open_logs_folder()
+            if status:
+                collect_logs.open_logs_folder()
         action = self.register_tools_menu_action("&Zip Logs")
         action.triggered.connect(save_logs)
 
diff --git a/Engine/Plugins/VirtualProduction/Switchboard/Source/Switchboard/switchboard/util/collect_logs.py b/Engine/Plugins/VirtualProduction/Switchboard/Source/Switchboard/switchboard/util/collect_logs.py
--- a/Engine/Plugins/VirtualProduction/Switchboard/Source/Switchboard/switchboard/util/collect_logs.py
+++ b/Engine/Plugins/VirtualProduction/Switchboard/Source/Switchboard/switchboard/util/collect_logs.py
@@ -62,11 +62,14 @@ def execute_zip_logs_workflow(config: Config, devices: typing.List[DeviceUnreal]
         return file_path
 
     if should_skip_because_running_devices():
-        return
+        return False
 
     zip_destination = ask_log_zip_name()
     if zip_destination is not None:
         zip_logs(zip_destination=zip_destination, config=config, devices=devices)
+        return True
+
+    return False
 
 
 def zip_logs(zip_destination: str, config: Config, devices: typing.List[DeviceUnreal]):
