Allow for new configs to be saved outside the config folder.

#skipci

diff --git a/Engine/Plugins/VirtualProduction/Switchboard/Source/Switchboard/switchboard/add_config_dialog.py b/Engine/Plugins/VirtualProduction/Switchboard/Source/Switchboard/switchboard/add_config_dialog.py
--- a/Engine/Plugins/VirtualProduction/Switchboard/Source/Switchboard/switchboard/add_config_dialog.py
+++ b/Engine/Plugins/VirtualProduction/Switchboard/Source/Switchboard/switchboard/add_config_dialog.py
@@ -520,17 +520,13 @@ class AddConfigDialog(QtWidgets.QDialog):
         if not config_path_str:
             return
 
-        # Protect against selecting a file outside of the configs directory.
+        # First attempt to get a relative path then allow for an absolute
+        #
         try:
             config_path_str = str(
                 config.get_relative_config_path(config_path_str))
         except Exception as e:
-            error_msg = QtWidgets.QMessageBox(QtWidgets.QMessageBox.Warning,
-                'Config Path Error', 'Cannot save config file.',
-                QtWidgets.QMessageBox.Ok, parent=self)
-            error_msg.setInformativeText(str(e))
-            error_msg.exec_()
-            return
+            config_path_str = str(config.get_absolute_config_path(config_path_str))
 
         self.config_path_line_edit.setText(config_path_str)
 
