Switchboard: File filter in open config dialog now space separates instead of semicolon.

#skipci

diff --git a/Engine/Plugins/VirtualProduction/Switchboard/Source/Switchboard/switchboard/devices/ndisplay/plugin_ndisplay.py b/Engine/Plugins/VirtualProduction/Switchboard/Source/Switchboard/switchboard/devices/ndisplay/plugin_ndisplay.py
--- a/Engine/Plugins/VirtualProduction/Switchboard/Source/Switchboard/switchboard/devices/ndisplay/plugin_ndisplay.py
+++ b/Engine/Plugins/VirtualProduction/Switchboard/Source/Switchboard/switchboard/devices/ndisplay/plugin_ndisplay.py
@@ -146,7 +146,7 @@ class AddnDisplayDialog(AddDeviceDialog):
 
         cfg_path, _ = QtWidgets.QFileDialog.getOpenFileName(
             self, "Select nDisplay config file", start_path,
-            "nDisplay Config (*.ndisplay;*.uasset)")
+            "nDisplay Config (*.ndisplay *.uasset)")
 
         if len(cfg_path) > 0 and os.path.exists(cfg_path):
             self.cbConfigs.setCurrentText(cfg_path)
