Switchboard: Extend network addressing to support hostnames in addition to IP addresses.

Address-type settings fields will now attempt name resolution on the entered value, with highlighting to indicate the lookup status (yellow = pending, red = failed).

diff --git a/Engine/Plugins/VirtualProduction/Switchboard/Source/Switchboard/switchboard/config.py b/Engine/Plugins/VirtualProduction/Switchboard/Source/Switchboard/switchboard/config.py
--- a/Engine/Plugins/VirtualProduction/Switchboard/Source/Switchboard/switchboard/config.py
+++ b/Engine/Plugins/VirtualProduction/Switchboard/Source/Switchboard/switchboard/config.py
@@ -9,6 +9,7 @@ import shutil
 import socket
 import sys
 import typing
+from typing import Type
 
 from PySide2 import QtCore
 from PySide2 import QtGui
@@ -16,7 +17,8 @@ from PySide2 import QtWidgets
 
 from switchboard import switchboard_widgets as sb_widgets
 from switchboard.switchboard_logging import LOGGER
-from switchboard.switchboard_widgets import DropDownMenuComboBox
+from switchboard.switchboard_widgets import (DropDownMenuComboBox,
+    NonScrollableComboBox)
 
 ROOT_CONFIGS_PATH = pathlib.Path(__file__).parent.with_name('configs')
 CONFIG_SUFFIX = '.json'
@@ -283,9 +285,8 @@ class Setting(QtCore.QObject):
         )
         
         # Clear the widget when it is destroyed to avoid dangling references
-        top_level_widget.destroyed.connect(
-            lambda destroyed_object=None, on_setting_changed_lambda=on_setting_changed_lambda, override_device_name=override_device_name:
-                self._on_widget_destroyed(on_setting_changed_lambda, override_device_name)
+        top_level_widget.destroyed.connect(lambda destroyed_object=None:
+            self._on_widget_destroyed(on_setting_changed_lambda, override_device_name)
         )
         
     def _on_widget_destroyed(self, on_setting_changed_lambda, override_device_name: str):
@@ -734,8 +735,11 @@ class OptionSetting(Setting):
 
         self.possible_values = possible_values or []
 
-    def _create_widgets(self, override_device_name: typing.Optional[str] = None) -> sb_widgets.NonScrollableComboBox:
-        combo = sb_widgets.NonScrollableComboBox()
+    def _create_widgets(
+        self, override_device_name: typing.Optional[str] = None, *,
+        widget_class: Type[NonScrollableComboBox] = NonScrollableComboBox
+    ) -> NonScrollableComboBox:
+        combo = widget_class()
         if self.tool_tip:
             combo.setToolTip(self.tool_tip)
 
@@ -745,8 +749,9 @@ class OptionSetting(Setting):
         combo.setCurrentIndex(
             combo.findData(self.get_value(override_device_name)))
 
-        self.set_widget(
-            widget=combo, override_device_name=override_device_name)
+        self.set_widget(widget=combo,
+                        override_device_name=override_device_name)
+
         combo.currentIndexChanged.connect(
             lambda index, override_device_name=override_device_name, combo=combo:
                 self._on_widget_value_changed(
@@ -762,9 +767,11 @@ class OptionSetting(Setting):
             return
 
         old_value = widget.currentText()
-        new_str_value = new_value if isinstance(new_value, str) else str(new_value)
+        new_str_value = str(new_value)
         if new_str_value != old_value:
-            widget.setCurrentIndex(widget.findText(new_str_value))
+            index = widget.findText(new_str_value)
+            if index != -1:
+                widget.setCurrentIndex(index)
 
 
 class MultiOptionSetting(OptionSetting):
@@ -1699,7 +1706,7 @@ class LoggingSetting(Setting):
             widget.model().category_verbosities = new_value
 
 
-class IPAddressSetting(OptionSetting):
+class AddressSetting(OptionSetting):
     def __init__(
         self,
         attr_name,
@@ -1710,73 +1717,56 @@ class IPAddressSetting(OptionSetting):
         allow_reset=True,
         migrate_data=None
     ):
-        from socket import getaddrinfo, AF_INET, gethostname
-        ip_addresses = [ip[4][0].__str__() for ip in getaddrinfo(host=gethostname(), port=None, family=AF_INET)]
-        ip_addresses.append("127.0.0.1")
         super().__init__(
             attr_name=attr_name,
             nice_name=nice_name,
             value=value,
-            possible_values=ip_addresses,
+            possible_values=list(self.generate_possible_addresses()),
             tool_tip=tool_tip,
             show_ui=show_ui,
             allow_reset=allow_reset,
-            migrate_data=migrate_data)#
-        
-        self._finishedEditingFired = False
+            migrate_data=migrate_data)
+
+    def _create_widgets(
+        self, override_device_name: typing.Optional[str] = None
+    ) -> NonScrollableComboBox:
+        combo: NonScrollableComboBox = super()._create_widgets(
+            override_device_name, widget_class=sb_widgets.AddressComboBox)
 
-    def _create_widgets(self, override_device_name: typing.Optional[str] = None) -> sb_widgets.NonScrollableComboBox:
-        combo: sb_widgets.NonScrollableComboBox = super()._create_widgets(override_device_name)
-        combo.setEditable(True)
         combo.setInsertPolicy(QtWidgets.QComboBox.InsertPolicy.NoInsert)
+
+        cur_value = self.get_value(override_device_name)
+        if combo.findText(cur_value, QtCore.Qt.MatchFlag.MatchExactly) == -1:
+            combo.addItem(str(cur_value), cur_value)
+            combo.setCurrentIndex(combo.findText(cur_value))
+
         combo.lineEdit().editingFinished.connect(
-            lambda combo=combo: self._validate_and_commit_ip(combo, override_device_name)
+            lambda: self._validate_and_commit_address(combo,
+                                                      override_device_name)
         )
-        
+
         return combo
-    
-    def _validate_and_commit_ip(self, combo: sb_widgets.NonScrollableComboBox, override_device_name:str):
-        def is_valid_ip(ip_str:str):
-            try:
-                import ipaddress
-                ip = ipaddress.ip_address(ip_str)
-                return True
-            except:
-                return False
-            
-        def perform_validation():
-            ip_str = combo.lineEdit().text()
-            if not is_valid_ip(ip_str):
-                answer = QtWidgets.QMessageBox.question(
-                    None,
-                    'Invalid IP',
-                    f'The IP \"{ip_str}\" seems to be invalid.\nDo you want to use this IP anyway?',
-                    QtWidgets.QMessageBox.Yes | QtWidgets.QMessageBox.No
-                )
-                if answer == QtWidgets.QMessageBox.No:
-                    combo.blockSignals(True)
-                    combo.lineEdit().blockSignals(True)
-
-                    current_value = self.get_value(override_device_name)
-                    index = combo.findData(current_value)
-                    is_combo_box_entry = index != -1
-                    if is_combo_box_entry:
-                        combo.setCurrentIndex(index)
-                    else:
-                        combo.lineEdit().setText(current_value)
 
-                    combo.lineEdit().blockSignals(False)
-                    combo.blockSignals(False)
-                    return
+    def _validate_and_commit_address(
+            self, combo: NonScrollableComboBox, override_device_name: str):
+        address_str = combo.lineEdit().text()
+        self._on_widget_value_changed(address_str, override_device_name=override_device_name)
+
+    def generate_possible_addresses(self):
+        return set[str]()
+
+
+class LocalAddressSetting(AddressSetting):
+    def generate_possible_addresses(self):
+        addresses = set[str]()
+        for address in socket.getaddrinfo(socket.gethostname(), None,
+                                          socket.AF_INET):
+            addresses.add(str(address[4][0]))
+        addresses.add('127.0.0.1')
+        addresses.add('localhost')
+        addresses.add(socket.gethostname())
+        return addresses
 
-            self._on_widget_value_changed(ip_str, override_device_name=override_device_name)
-        
-        if self._finishedEditingFired:
-            return
-        else:
-            self._finishedEditingFired = True
-            perform_validation()
-            self._finishedEditingFired = False
 
 class ConfigPathError(Exception):
     '''
@@ -1905,9 +1895,10 @@ class Config(object):
                 self.file_path = get_absolute_config_path(file_path)
 
                 # Read the json config file
-                with open(self.file_path) as f:
+                with open(self.file_path, 'r') as f:
                     LOGGER.debug(f'Loading Config {self.file_path}')
-                    data = json.load(f)
+                    data = json.load(f)                    
+                        
             except (ConfigPathError, FileNotFoundError) as e:
                 LOGGER.error(f'Config: {e}')
                 self.file_path = None
@@ -1956,16 +1947,22 @@ class Config(object):
         # into the kwargs.
         for device_type, devices in data.get('devices', {}).items():
             for device_name, data in devices.items():
-                if device_name == "settings":
+                if device_name == 'settings':
                     self._plugin_data_from_config[device_type] = data
                 else:
-                    ip_address = data["ip_address"]
+                    # Migrate ip_address -> address
+                    if 'ip_address' in data:
+                        address = data['ip_address']
+                        del data['ip_address']
+                    else:
+                        address = data['address']
+
                     device_data = {
-                        "name": device_name,
-                        "ip_address": ip_address
+                        'name': device_name,
+                        'address': address
                     }
-                    device_data["kwargs"] = {
-                        k: v for (k, v) in data.items() if k != "ip_address"}
+                    device_data['kwargs'] = {
+                        k: v for (k, v) in data.items() if k != 'address'}
                     self._device_data_from_config.setdefault(
                         device_type, []).append(device_data)
 
@@ -2174,7 +2171,7 @@ class Config(object):
                 value=data.get('muserver_multicast_endpoint', '230.0.0.1:6666'),
                 tool_tip=(
                     'Multicast group and port (-UDPMESSAGING_TRANSPORT_MULTICAST) '
-                    'in the {ip}:{port} endpoint format. The multicast group IP '
+                    'in the {address}:{port} endpoint format. The multicast group address '
                     'must be in the range 224.0.0.0 to 239.255.255.255.'),
             ),
             "multiuser_exe": StringSetting(
@@ -2509,15 +2506,15 @@ class Config(object):
             self.ENGINE_DIR.get_value(), self.LISTENER_EXE.get_value())
 
     # todo-dara: find a way to do this directly in the LiveLinkFace plugin code
-    def unreal_device_ip_addresses(self):
-        unreal_ips = []
+    def unreal_device_addresses(self):
+        unreal_addresses = []
         for (device_type, device_name), (settings, overrides) in \
                 self._device_settings.items():
             if device_type == "Unreal":
                 for setting in settings:
-                    if setting.attr_name == "ip_address":
-                        unreal_ips.append(setting.get_value(device_name))
-        return unreal_ips
+                    if setting.attr_name == "address":
+                        unreal_addresses.append(setting.get_value(device_name))
+        return unreal_addresses
 
     @staticmethod
     def engine_exe_path(engine_dir: str, exe_basename: str):
@@ -2596,11 +2593,11 @@ class UserSettings(object):
             config_paths = list_config_paths()
             self.CONFIG = config_paths[0] if config_paths else None
 
-        # IP Address of the machine running Switchboard
-        self.IP_ADDRESS = IPAddressSetting(
-            "ip_address",
-            "IP Address",
-            data.get("ip_address", socket.gethostbyname(socket.gethostname()))
+        # Address of the machine running Switchboard
+        self.ADDRESS = LocalAddressSetting(
+            "address",
+            "Address",
+            data.get("address", socket.gethostbyname(socket.gethostname()))
         )
         self.TRANSPORT_PATH = FilePathSetting(
                 "transport_path",
@@ -2624,7 +2621,7 @@ class UserSettings(object):
     def save(self):
         data = {
             'config': '',
-            'ip_address': self.IP_ADDRESS.get_value(),
+            'address': self.ADDRESS.get_value(),
             'transport_path': self.TRANSPORT_PATH.get_value(),
             'muserver_session_name': self.MUSERVER_SESSION_NAME,
             'current_sequence': self.CURRENT_SEQUENCE,
diff --git a/Engine/Plugins/VirtualProduction/Switchboard/Source/Switchboard/switchboard/device_list_widget.py b/Engine/Plugins/VirtualProduction/Switchboard/Source/Switchboard/switchboard/device_list_widget.py
--- a/Engine/Plugins/VirtualProduction/Switchboard/Source/Switchboard/switchboard/device_list_widget.py
+++ b/Engine/Plugins/VirtualProduction/Switchboard/Source/Switchboard/switchboard/device_list_widget.py
@@ -325,7 +325,7 @@ class DeviceWidgetHeader(QtWidgets.QWidget):
         super().__init__(parent)
 
         self.name = name
-        self.ip_address_label = None
+        self.address_label = None
         self.device_hash = 0 # When list widget is looking for devices it checks the device_has
 
         self.layout = QtWidgets.QHBoxLayout()
@@ -339,10 +339,10 @@ class DeviceWidgetHeader(QtWidgets.QWidget):
             return label
 
         self.name_label = __label(self.name + " Devices")
-        self.ip_address_label = __label('IP Address')
+        self.address_label = __label('Address')
         self.changelist_label = __label('Changelist') if show_changelist else None
         self.layout.addWidget(self.name_label)
-        self.layout.addWidget(self.ip_address_label)
+        self.layout.addWidget(self.address_label)
 
         spacer = QtWidgets.QSpacerItem(0, 20, QtWidgets.QSizePolicy.Expanding, QtWidgets.QSizePolicy.Minimum)
         self.layout.addItem(spacer)
@@ -385,32 +385,32 @@ class DeviceWidgetHeader(QtWidgets.QWidget):
             self.layout.addWidget(self.connect_button)
 
         self.name_label.setMinimumSize(QtCore.QSize(235, 0))
-        self.ip_address_label.setMinimumSize(QtCore.QSize(100, 0))
+        self.address_label.setMinimumSize(QtCore.QSize(100, 0))
 
     def showEvent(self, event):
         super().showEvent(event)
 
-        self.find_ip_address_label()
+        self.find_address_label()
 
-    def find_ip_address_label(self):
+    def find_address_label(self):
         labels = self.findChildren(QtWidgets.QLabel)
         for label in labels:
-            if label.text() == 'IP Address':
-                self.ip_address_label = label
+            if label.text() == 'Address':
+                self.address_label = label
                 return
 
     def resizeEvent(self, event):
         super().resizeEvent(event)
 
-        if not self.ip_address_label:
+        if not self.address_label:
             return
 
         width = event.size().width()
 
-        if width < sb_widgets.DEVICE_WIDGET_HIDE_IP_ADDRESS_WIDTH:
-            self.ip_address_label.hide()
+        if width < sb_widgets.DEVICE_WIDGET_HIDE_ADDRESS_WIDTH:
+            self.address_label.hide()
         else:
-            self.ip_address_label.show()
+            self.address_label.show()
 
     def paintEvent(self, event):
         opt = QtWidgets.QStyleOption()
diff --git a/Engine/Plugins/VirtualProduction/Switchboard/Source/Switchboard/switchboard/devices/device_base.py b/Engine/Plugins/VirtualProduction/Switchboard/Source/Switchboard/switchboard/devices/device_base.py
--- a/Engine/Plugins/VirtualProduction/Switchboard/Source/Switchboard/switchboard/devices/device_base.py
+++ b/Engine/Plugins/VirtualProduction/Switchboard/Source/Switchboard/switchboard/devices/device_base.py
@@ -16,7 +16,7 @@ import pythonosc.udp_client
 
 from switchboard import config_osc as osc
 from switchboard import recording
-from switchboard.config import CONFIG, BoolSetting, StringSetting
+from switchboard.config import CONFIG, BoolSetting, AddressSetting
 from switchboard.switchboard_logging import LOGGER
 
 if TYPE_CHECKING:
@@ -45,6 +45,7 @@ class PluginHeaderWidgets(IntFlag):
     CHANGELIST_LABEL = 4,
     AUTOJOIN_MU = 8
 
+
 class DeviceQtHandler(QtCore.QObject):
     signal_device_status_changed = QtCore.Signal(object, int)
     signal_device_connect_failed = QtCore.Signal(object)
@@ -76,14 +77,13 @@ class Device(QtCore.QObject):
         )
     }
 
-    def __init__(self, name: str, ip_address: str, **kwargs):
+    def __init__(self, name: str, address: str, **kwargs):
         super().__init__()
 
         self._name = name  # Assigned name
         self.device_qt_handler = DeviceQtHandler()
 
-        self.setting_ip_address = StringSetting(
-            "ip_address", 'IP address', ip_address)
+        self.setting_address = AddressSetting('address', 'Address', address)
 
         # override any setting that was passed via kwargs
         for setting in self.setting_overrides():
@@ -101,7 +101,7 @@ class Device(QtCore.QObject):
         self.auto_connect = True
 
         # TODO: This is not really a hash. Could be changed to real hash based
-        # on e.g. ip and name.
+        # on e.g. address and name.
         self.device_hash = random.getrandbits(64)
 
         # Lazily create the OSC client as needed
@@ -113,14 +113,14 @@ class Device(QtCore.QObject):
 
     def init(self, widget_class, icons):
         self.widget = widget_class(
-            self.name, self.device_hash, self.ip_address, icons)
+            self.name, self.device_hash, self.address, icons)
         self.widget.signal_device_name_changed.connect(
             self.on_device_name_changed)
-        self.widget.signal_ip_address_changed.connect(
-            self.on_ip_address_changed)
-        self.setting_ip_address.signal_setting_changed.connect(
-            lambda _, new_ip, widget=self.widget:
-                widget.on_ip_address_changed(new_ip))
+        self.widget.signal_address_changed.connect(
+            self.on_address_changed)
+        self.setting_address.signal_setting_changed.connect(
+            lambda _, new_address, widget=self.widget:
+                widget.on_address_changed(new_address))
 
         # Let the CONFIG class know what settings/properties this device has
         # so it can save the config file on changes.
@@ -134,9 +134,9 @@ class Device(QtCore.QObject):
             self.name = new_name
             CONFIG.on_device_name_changed(old_name, new_name)
 
-    def on_ip_address_changed(self, new_ip):
-        if self.ip_address != new_ip:
-            self.ip_address = new_ip
+    def on_address_changed(self, new_address):
+        if self.address != new_address:
+            self.address = new_address
 
     @property
     def device_type(self):
@@ -168,7 +168,7 @@ class Device(QtCore.QObject):
 
     def device_settings(self):
         # settings that are specific to an instance of the device
-        return [self.setting_ip_address]
+        return [self.setting_address]
 
     def setting_overrides(self):
         # All settings that a device may override. these might be project or
@@ -207,14 +207,14 @@ class Device(QtCore.QObject):
             self, value)
 
     @property
-    def ip_address(self) -> str:
-        return self.setting_ip_address.get_value()
+    def address(self) -> str:
+        return self.setting_address.get_value()
 
-    @ip_address.setter
-    def ip_address(self, value: str):
-        self.setting_ip_address.update_value(value)
+    @address.setter
+    def address(self, value: str):
+        self.setting_address.update_value(value)
         # todo-dara: probably better to have the osc client connect to a
-        # change of the ip address.
+        # change of the address.
         self.setup_osc_client(CONFIG.OSC_CLIENT_PORT.get_value())
 
     @property
@@ -307,8 +307,12 @@ class Device(QtCore.QObject):
         self.status = DeviceStatus.DISCONNECTED
 
     def setup_osc_client(self, osc_port):
-        self.osc_client = pythonosc.udp_client.SimpleUDPClient(
-            self.ip_address, osc_port)
+        try:
+            self.osc_client = pythonosc.udp_client.SimpleUDPClient(
+                self.address, osc_port)
+        except socket.gaierror as exc:
+            LOGGER.error(f'{self.name}: Invalid OSC server address',
+                         exc_info=exc)
 
     def send_osc_message(self, command, value, log=True):
         if not self.osc_client:
@@ -317,13 +321,13 @@ class Device(QtCore.QObject):
         if log:
             LOGGER.osc(
                 f'OSC Server: Sending {command} {value} '
-                f'to {self.name} ({self.ip_address})')
+                f'to {self.name} ({self.address})')
 
         try:
             self.osc_client.send_message(command, value)
         except socket.gaierror as e:
             LOGGER.error(
-                f'{self.name}: Incorrect ip address when sending '
+                f'{self.name}: Incorrect address when sending '
                 f'OSC message. {e}')
 
     def record_start(self, slate, take, description):
diff --git a/Engine/Plugins/VirtualProduction/Switchboard/Source/Switchboard/switchboard/devices/device_manager.py b/Engine/Plugins/VirtualProduction/Switchboard/Source/Switchboard/switchboard/devices/device_manager.py
--- a/Engine/Plugins/VirtualProduction/Switchboard/Source/Switchboard/switchboard/devices/device_manager.py
+++ b/Engine/Plugins/VirtualProduction/Switchboard/Source/Switchboard/switchboard/devices/device_manager.py
@@ -8,9 +8,7 @@ from typing import Dict, List, Type
 
 from PySide2 import QtCore, QtGui
 
-from switchboard.config import CONFIG, SETTINGS
 from switchboard.switchboard_logging import LOGGER
-from switchboard import config_osc as osc
 from switchboard.devices.device_base import Device, DeviceStatus
 from switchboard.devices.device_widget_base import AddDeviceDialog, DeviceWidget
 import switchboard.switchboard_utils as utils
@@ -36,19 +34,19 @@ class DeviceManager(QtCore.QObject):
     def add_devices(self, device_config):
         for device_type, devices in device_config.items():
             for device in devices:
-                self._add_device(device_type, device["name"], device["ip_address"], **device["kwargs"])
+                self._add_device(device_type, device["name"], device["address"], **device["kwargs"])
 
         # notify the plugins that all devices were added (gets called even if no devices were added)
         for device_cls in self._plugins.values():
             device_cls.all_devices_added()
 
-    def _add_device(self, device_type, name, ip, **kwargs):
+    def _add_device(self, device_type, name, address, **kwargs):
         device_cls_name = device_type
         if device_cls_name not in self._plugins:
             LOGGER.error(f"Could not find plugin for {device_type} in {DEVICE_PLUGIN_PACKAGE}")
             return None
 
-        device = self._plugins[device_cls_name](name, ip, **kwargs)
+        device = self._plugins[device_cls_name](name, address, **kwargs)
         widget_class = self._plugin_widgets[device_cls_name]
         icons = self._plugin_icons[device_cls_name] if device_cls_name in self._plugin_icons.keys() else None
         device.init(widget_class, icons)
@@ -189,9 +187,9 @@ class DeviceManager(QtCore.QObject):
 
         return self._devices[device_hash]
 
-    def device_with_ip_address(self, ip_address):
+    def device_with_address(self, address):
         for device in self._devices.values():
-            if device.ip_address == ip_address:
+            if device.address == address:
                 return device
         return None
 
diff --git a/Engine/Plugins/VirtualProduction/Switchboard/Source/Switchboard/switchboard/devices/device_widget_base.py b/Engine/Plugins/VirtualProduction/Switchboard/Source/Switchboard/switchboard/devices/device_widget_base.py
--- a/Engine/Plugins/VirtualProduction/Switchboard/Source/Switchboard/switchboard/devices/device_widget_base.py
+++ b/Engine/Plugins/VirtualProduction/Switchboard/Source/Switchboard/switchboard/devices/device_widget_base.py
@@ -9,9 +9,6 @@ from switchboard.devices.device_base import DeviceStatus
 import switchboard.switchboard_widgets as sb_widgets
 
 
-ip_regex = QtCore.QRegExp("^\\s*((([01]?[0-9]?[0-9]|2[0-4][0-9]|25[0-5])\\.){3}([01]?[0-9]?[0-9]|2[0-4][0-9]|25[0-5]))\\s*$")
-ip_validator = QtGui.QRegExpValidator(ip_regex)
-
 class DeviceWidgetItem(QtWidgets.QWidget):
     """
     Custom class to get QSS working correctly to achieve a look.
@@ -26,6 +23,7 @@ class DeviceWidgetItem(QtWidgets.QWidget):
         painter = QtGui.QPainter(self)
         self.style().drawPrimitive(QtWidgets.QStyle.PE_Widget, opt, painter, self)
 
+
 class DeviceAutoJoinMUServerUI(QtCore.QObject):
     signal_device_widget_autojoin_mu = QtCore.Signal(object)
 
@@ -81,12 +79,14 @@ class DeviceWidget(QtWidgets.QWidget):
     signal_device_widget_trigger_stop_toggled = QtCore.Signal(object, bool)
 
     signal_device_name_changed = QtCore.Signal(str)
-    signal_ip_address_changed = QtCore.Signal(str)
+    signal_address_changed = QtCore.Signal(str)
+
+    hostname_validator = sb_widgets.HostnameValidator()
 
-    def __init__(self, name, device_hash, ip_address, icons, parent=None):
+    def __init__(self, name, device_hash, address, icons, parent=None):
         super().__init__(parent)
 
-        # Lookup device by a hash instead of name/ip_address
+        # Lookup device by a hash instead of name/address
         self.device_hash = device_hash
         self.icons = icons
 
@@ -118,14 +118,14 @@ class DeviceWidget(QtWidgets.QWidget):
         self.name_line_edit.setMaximumSize(QtCore.QSize(150, 40))
         # 20 + 11 + 60 + 150
 
-        # IP Address Label
-        self.ip_address_line_edit = FramelessQLineEdit()
-        self.ip_address_line_edit.setObjectName('device_address')
-        self.ip_address_line_edit.setValidator(ip_validator)
-        self.ip_address_line_edit.editingFinished.connect(self.on_ip_address_edited)
-        self.ip_address_line_edit.setText(ip_address)
-        self.ip_address_line_edit.setAlignment(QtCore.Qt.AlignCenter)
-        self.ip_address_line_edit.setMaximumSize(QtCore.QSize(100, 40))
+        # Address Label
+        self.address_line_edit = FramelessQLineEdit()
+        self.address_line_edit.setObjectName('device_address')
+        self.address_line_edit.setValidator(DeviceWidget.hostname_validator)
+        self.address_line_edit.editingFinished.connect(self.on_address_edited)
+        self.address_line_edit.setText(address)
+        self.address_line_edit.setAlignment(QtCore.Qt.AlignCenter)
+        self.address_line_edit.setMaximumSize(QtCore.QSize(100, 40))
 
         # Create a widget where the body of the item will go
         # This is made to allow the edit buttons to sit "outside" of the item
@@ -144,7 +144,7 @@ class DeviceWidget(QtWidgets.QWidget):
         self.add_widget_to_layout(self.status_icon)
         self.add_widget_to_layout(self.device_icon)
         self.add_widget_to_layout(self.name_line_edit)
-        self.add_widget_to_layout(self.ip_address_line_edit)
+        self.add_widget_to_layout(self.address_line_edit)
 
         spacer = QtWidgets.QSpacerItem(0, 20, QtWidgets.QSizePolicy.Expanding, QtWidgets.QSizePolicy.Minimum)
         self.add_item_to_layout(spacer)
@@ -153,7 +153,7 @@ class DeviceWidget(QtWidgets.QWidget):
         #self.previous_status = DeviceStatus.DISCONNECTED
 
         # Set style as disconnected
-        for label in [self.name_line_edit, self.ip_address_line_edit]:
+        for label in [self.name_line_edit, self.address_line_edit]:
             sb_widgets.set_qt_property(label, 'disconnected', True)
 
         # Store the control buttons by name ("connect", "open", etc.)
@@ -213,16 +213,16 @@ class DeviceWidget(QtWidgets.QWidget):
 
             self.signal_device_name_changed.emit(new_value)
 
-    def on_ip_address_edited(self):
-        new_value = self.ip_address_line_edit.text()
+    def on_address_edited(self):
+        new_value = self.address_line_edit.text()
 
-        if self.ip_address_line_edit.is_valid and self.ip_address_line_edit.current_text != new_value:
-            sb_widgets.set_qt_property(self.ip_address_line_edit, "input_error", False)
+        if self.address_line_edit.is_valid and self.address_line_edit.current_text != new_value:
+            sb_widgets.set_qt_property(self.address_line_edit, "input_error", False)
 
-            self.signal_ip_address_changed.emit(new_value)
+            self.signal_address_changed.emit(new_value)
 
-    def on_ip_address_changed(self, new_ip):
-        self.ip_address_line_edit.setText(new_ip)
+    def on_address_changed(self, new_address):
+        self.address_line_edit.setText(new_address)
 
     def _add_control_buttons(self):
         pass
@@ -250,31 +250,31 @@ class DeviceWidget(QtWidgets.QWidget):
 
         # Device icon
         if status in {DeviceStatus.DISCONNECTED, DeviceStatus.CONNECTING}:
-            for label in [self.name_line_edit, self.ip_address_line_edit]:
+            for label in [self.name_line_edit, self.address_line_edit]:
                 sb_widgets.set_qt_property(label, 'disconnected', True)
 
             pixmap = self.icon_for_state("disabled").pixmap(QtCore.QSize(40, 40))
             self.device_icon.setPixmap(pixmap)
 
             if status == DeviceStatus.DISCONNECTED:
-                # Make the Name and IP editable when disconnected.
+                # Make the name and address editable when disconnected.
                 self.name_line_edit.setReadOnly(False)
-                self.ip_address_line_edit.setReadOnly(False)
+                self.address_line_edit.setReadOnly(False)
             elif status == DeviceStatus.CONNECTING:
-                # Make the Name and IP non-editable while connecting.
+                # Make the name and address non-editable while connecting.
                 self.name_line_edit.setReadOnly(True)
-                self.ip_address_line_edit.setReadOnly(True)
+                self.address_line_edit.setReadOnly(True)
         elif ((previous_status in {DeviceStatus.DISCONNECTED, DeviceStatus.CONNECTING}) and
                 status > DeviceStatus.CONNECTING):
-            for label in [self.name_line_edit, self.ip_address_line_edit]:
+            for label in [self.name_line_edit, self.address_line_edit]:
                 sb_widgets.set_qt_property(label, 'disconnected', False)
 
             pixmap = self.icon_for_state("enabled").pixmap(QtCore.QSize(40, 40))
             self.device_icon.setPixmap(pixmap)
 
-            # Make the Name and IP non-editable when connected.
+            # Make the name and address non-editable when connected.
             self.name_line_edit.setReadOnly(True)
-            self.ip_address_line_edit.setReadOnly(True)
+            self.address_line_edit.setReadOnly(True)
 
         # Handle coloring List Widget items if they are recording
         if status == DeviceStatus.RECORDING:
@@ -287,10 +287,10 @@ class DeviceWidget(QtWidgets.QWidget):
 
         width = event.size().width()
 
-        if width < sb_widgets.DEVICE_WIDGET_HIDE_IP_ADDRESS_WIDTH:
-            self.ip_address_line_edit.hide()
+        if width < sb_widgets.DEVICE_WIDGET_HIDE_ADDRESS_WIDTH:
+            self.address_line_edit.hide()
         else:
-            self.ip_address_line_edit.show()
+            self.address_line_edit.show()
 
     def assign_button_to_name(self, name, button):
         if name:
@@ -318,12 +318,12 @@ class AddDeviceDialog(QtWidgets.QDialog):
 
         self.name_field = QtWidgets.QLineEdit(self)
 
-        self.ip_field = QtWidgets.QLineEdit(self)
-        self.ip_field.setValidator(ip_validator)
+        self.address_field = QtWidgets.QLineEdit(self)
+        self.address_field.setValidator(DeviceWidget.hostname_validator)
 
         self.form_layout = QtWidgets.QFormLayout()
         self.form_layout.addRow("Name", self.name_field)
-        self.form_layout.addRow("IP Address", self.ip_field)
+        self.form_layout.addRow("Address", self.address_field)
 
         layout = QtWidgets.QVBoxLayout()
         layout.insertLayout(0, self.form_layout)
@@ -340,7 +340,7 @@ class AddDeviceDialog(QtWidgets.QDialog):
             self.name_field.setValidator(validator)
 
     def devices_to_add(self):
-        return [{"type": self.device_type, "name": self.name_field.text(), "ip_address": self.ip_field.text(), "kwargs": {}}]
+        return [{"type": self.device_type, "name": self.name_field.text(), "address": self.address_field.text(), "kwargs": {}}]
 
     def devices_to_remove(self):
         return []
diff --git a/Engine/Plugins/VirtualProduction/Switchboard/Source/Switchboard/switchboard/devices/kipro/plugin_kipro.py b/Engine/Plugins/VirtualProduction/Switchboard/Source/Switchboard/switchboard/devices/kipro/plugin_kipro.py
--- a/Engine/Plugins/VirtualProduction/Switchboard/Source/Switchboard/switchboard/devices/kipro/plugin_kipro.py
+++ b/Engine/Plugins/VirtualProduction/Switchboard/Source/Switchboard/switchboard/devices/kipro/plugin_kipro.py
@@ -37,11 +37,11 @@ def unresponsive_kipro(f):
 
 
 class DeviceKiPro(Device):
-    def __init__(self, name, ip_address, **kwargs):
+    def __init__(self, name, address, **kwargs):
         self.setting_auto_play = BoolSetting(
             "auto_play", "Auto Play After Stop", False)
 
-        super().__init__(name, ip_address, **kwargs)
+        super().__init__(name, address, **kwargs)
 
         self.client = None
 
@@ -69,7 +69,7 @@ class DeviceKiPro(Device):
     # TODO: Make sure KiPro is not in transport mode
     @unresponsive_kipro
     def connect_listener(self):
-        self.client = kipro.Client(f'http://{self.ip_address}')
+        self.client = kipro.Client(f'http://{self.address}')
 
         self.status = DeviceStatus.READY
 
@@ -237,7 +237,7 @@ class DeviceKiPro(Device):
     def transport_file(self, device_path, output_dir):
         self.client.setParameter('eParamID_MediaState', 1)
 
-        device_path = f'http://{self.ip_address}/{device_path}'
+        device_path = f'http://{self.address}/{device_path}'
         output_path = os.path.join(output_dir, os.path.basename(device_path))
 
         # Put all drives into file transfer mode
@@ -248,8 +248,8 @@ class DeviceWidgetKiPro(DeviceWidget):
     signal_device_widget_play = QtCore.Signal(object)
     signal_device_widget_stop = QtCore.Signal(object)
 
-    def __init__(self, name, device_hash, ip_address, icons, parent=None):
-        super().__init__(name, device_hash, ip_address, icons, parent=parent)
+    def __init__(self, name, device_hash, address, icons, parent=None):
+        super().__init__(name, device_hash, address, icons, parent=parent)
 
     def _add_control_buttons(self):
         super()._add_control_buttons()
diff --git a/Engine/Plugins/VirtualProduction/Switchboard/Source/Switchboard/switchboard/devices/livelinkface/plugin_livelinkface.py b/Engine/Plugins/VirtualProduction/Switchboard/Source/Switchboard/switchboard/devices/livelinkface/plugin_livelinkface.py
--- a/Engine/Plugins/VirtualProduction/Switchboard/Source/Switchboard/switchboard/devices/livelinkface/plugin_livelinkface.py
+++ b/Engine/Plugins/VirtualProduction/Switchboard/Source/Switchboard/switchboard/devices/livelinkface/plugin_livelinkface.py
@@ -13,8 +13,8 @@ import datetime
 class DeviceLiveLinkFace(Device):
     QUERY_TIME = 5
 
-    def __init__(self, name, ip_address, port=CONFIG.OSC_CLIENT_PORT.get_value(), **kwargs):
-        super().__init__(name, ip_address, **kwargs)
+    def __init__(self, name, address, port=CONFIG.OSC_CLIENT_PORT.get_value(), **kwargs):
+        super().__init__(name, address, **kwargs)
 
         self.auto_connect = False
         self.port = port
@@ -55,7 +55,7 @@ class DeviceLiveLinkFace(Device):
         super().disconnect_listener()
 
     def launch(self):
-        unreal_devices = CONFIG.unreal_device_ip_addresses()
+        unreal_devices = CONFIG.unreal_device_addresses()
 
         if not unreal_devices:
             LOGGER.warning(f'{self.name}: does not have an Unreal device to stream to. Please add one')
@@ -67,7 +67,7 @@ class DeviceLiveLinkFace(Device):
         self.send_osc_message(osc.ARSESSION_START, 1)
 
         self.send_osc_message(osc.CLEAR_LIVE_LINK_ADDRESSES, 1)
-        # Loop through all of the Unreal IP Addresses to add them
+        # Loop through all of the Unreal Addresses to add them
         for unreal_address in unreal_devices:
             self.send_osc_message(osc.ADD_LIVE_LINK_ADDRESS, [unreal_address, 11111])
 
@@ -89,7 +89,7 @@ class DeviceLiveLinkFace(Device):
             self.disconnect_listener()
 
         if self.look_for_device:
-            if len(CONFIG.unreal_device_ip_addresses()) > 0:
+            if len(CONFIG.unreal_device_addresses()) > 0:
                 if self.status > DeviceStatus.DISCONNECTED:
                     self.battery_query()
                 else:
@@ -105,9 +105,9 @@ class DeviceLiveLinkFace(Device):
     def battery_query(self):
         self.send_osc_message(osc.BATTERY_QUERY, 1, log=False)
 
-    # Tell the mobile which ip_address is running Switchboard
+    # Tell the mobile which address is running Switchboard
     def osc_add_send_target(self):
-        self.send_osc_message(osc.OSC_ADD_SEND_TARGET, [SETTINGS.IP_ADDRESS.get_value(), CONFIG.OSC_SERVER_PORT.get_value()], log=True)
+        self.send_osc_message(osc.OSC_ADD_SEND_TARGET, [SETTINGS.ADDRESS.get_value(), CONFIG.OSC_SERVER_PORT.get_value()], log=True)
 
     def osc_add_send_target_confirm(self):
         # If the mobile is left on and CP is relaunched the normal bootup sequence doesn't happen 
@@ -129,8 +129,8 @@ class DeviceLiveLinkFace(Device):
 
 
 class DeviceWidgetLiveLinkFace(DeviceWidget):
-    def __init__(self, name, device_hash, ip_address, icons, parent=None):
-        super().__init__(name, device_hash, ip_address, icons, parent=parent)
+    def __init__(self, name, device_hash, address, icons, parent=None):
+        super().__init__(name, device_hash, address, icons, parent=parent)
 
     def _add_control_buttons(self):
         super()._add_control_buttons()
diff --git a/Engine/Plugins/VirtualProduction/Switchboard/Source/Switchboard/switchboard/devices/motive/plugin_motive.py b/Engine/Plugins/VirtualProduction/Switchboard/Source/Switchboard/switchboard/devices/motive/plugin_motive.py
--- a/Engine/Plugins/VirtualProduction/Switchboard/Source/Switchboard/switchboard/devices/motive/plugin_motive.py
+++ b/Engine/Plugins/VirtualProduction/Switchboard/Source/Switchboard/switchboard/devices/motive/plugin_motive.py
@@ -36,8 +36,8 @@ class DeviceMotive(Device):
     setting_motive_port = IntSetting(
         "motive_port", "Motive Command Port", 1510)
 
-    def __init__(self, name, ip_address, **kwargs):
-        super().__init__(name, ip_address, **kwargs)
+    def __init__(self, name, address, **kwargs):
+        super().__init__(name, address, **kwargs)
 
         self.trigger_start = True
         self.trigger_stop = True
@@ -175,7 +175,7 @@ class DeviceMotive(Device):
                     self._flush_read_sockets()
                     self.socket.sendto(
                         message_bytes,
-                        (self.ip_address,
+                        (self.address,
                          self.setting_motive_port.get_value()))
 
                     read_sockets, _, _ = select.select(
@@ -296,8 +296,8 @@ class DeviceMotive(Device):
 
 
 class DeviceWidgetMotive(DeviceWidget):
-    def __init__(self, name, device_hash, ip_address, icons, parent=None):
-        super().__init__(name, device_hash, ip_address, icons, parent=parent)
+    def __init__(self, name, device_hash, address, icons, parent=None):
+        super().__init__(name, device_hash, address, icons, parent=parent)
 
     def _add_control_buttons(self):
         super()._add_control_buttons()
diff --git a/Engine/Plugins/VirtualProduction/Switchboard/Source/Switchboard/switchboard/devices/ndisplay/ndisplay_monitor.py b/Engine/Plugins/VirtualProduction/Switchboard/Source/Switchboard/switchboard/devices/ndisplay/ndisplay_monitor.py
--- a/Engine/Plugins/VirtualProduction/Switchboard/Source/Switchboard/switchboard/devices/ndisplay/ndisplay_monitor.py
+++ b/Engine/Plugins/VirtualProduction/Switchboard/Source/Switchboard/switchboard/devices/ndisplay/ndisplay_monitor.py
@@ -169,7 +169,7 @@ class nDisplayMonitor(QAbstractTableModel):
         for colname in self.colnames:
             data[colname] = self.DATA_MISSING
 
-        data['Host'] = str(device.ip_address)
+        data['Host'] = str(device.address)
         data['Node'] = device.name
         data['Connected'] = \
             'yes' if device.unreal_client.is_connected else 'no'
diff --git a/Engine/Plugins/VirtualProduction/Switchboard/Source/Switchboard/switchboard/devices/ndisplay/plugin_ndisplay.py b/Engine/Plugins/VirtualProduction/Switchboard/Source/Switchboard/switchboard/devices/ndisplay/plugin_ndisplay.py
--- a/Engine/Plugins/VirtualProduction/Switchboard/Source/Switchboard/switchboard/devices/ndisplay/plugin_ndisplay.py
+++ b/Engine/Plugins/VirtualProduction/Switchboard/Source/Switchboard/switchboard/devices/ndisplay/plugin_ndisplay.py
@@ -3,12 +3,11 @@
 import concurrent.futures
 import json
 import os
-import typing
 from pathlib import Path
 import socket
 import struct
 import traceback
-from typing import List, Optional, Tuple
+from typing import Optional
 
 from PySide2 import QtCore
 from PySide2 import QtWidgets
@@ -18,8 +17,8 @@ from switchboard import switchboard_utils as sb_utils
 from switchboard import switchboard_widgets as sb_widgets
 from switchboard import switchboard_dialog as sb_dialog
 from switchboard.config import CONFIG, BoolSetting, FilePathSetting, \
-    LoggingSetting, OptionSetting, Setting, StringSetting, SETTINGS, StringListSetting, \
-    migrate_comma_separated_string_to_list
+    LoggingSetting, OptionSetting, Setting, StringSetting, SETTINGS, \
+    StringListSetting, migrate_comma_separated_string_to_list
 from switchboard.devices.device_widget_base import AddDeviceDialog
 from switchboard.devices.unreal.plugin_unreal import DeviceUnreal, \
     DeviceWidgetUnreal
@@ -39,11 +38,11 @@ class AddnDisplayDialog(AddDeviceDialog):
         # Set enough width for a decent file path length
         self.setMinimumWidth(640)
 
-        # remove unneded base dialog layout rows
+        # remove unneeded base dialog layout rows
         self.form_layout.removeRow(self.name_field)
-        self.form_layout.removeRow(self.ip_field)
+        self.form_layout.removeRow(self.address_field)
         self.name_field = None
-        self.ip_field = None
+        self.address_field = None
 
         # Button to browse for supported config files
         self.btnBrowse = QtWidgets.QPushButton(self, text="Browse")
@@ -152,14 +151,14 @@ class AddnDisplayDialog(AddDeviceDialog):
             self.cbConfigs.setCurrentText(cfg_path)
             SETTINGS.LAST_BROWSED_PATH = os.path.dirname(cfg_path)
             SETTINGS.save()
-    
+
     def generate_short_unique_config_name(config_path: str, file_name: str) -> str:
         config_path = CONFIG.shrink_path(config_path)
         return sb_dialog.SwitchboardDialog.filter_empty_abiguated_path(config_path, file_name)
-    
+
     def on_clicked_btnFindConfigs(self):
         ''' Finds and populates config combobox '''
-        
+
         # We will look for config files in the project's Content folder
         configs_path = os.path.normpath(
             os.path.join(
@@ -209,7 +208,7 @@ class AddnDisplayDialog(AddDeviceDialog):
                 aparser = UassetParser(file, allowUnversioned=True)
 
                 for assetdata in aparser.aregdata:
-                    if assetdata.ObjectClassName in DisplayConfig.validClassNames:
+                    if assetdata.ObjectClassName in DisplayConfig.VALID_CLASS_NAMES:
                         return asset
 
             raise ValueError
@@ -260,8 +259,6 @@ class AddnDisplayDialog(AddDeviceDialog):
 
         # update the combo box with the items.
         self.recall_config_itemDatas()
-        
-
 
     def devices_to_add(self):
         cfg_file = self.current_config_path()
@@ -348,12 +345,16 @@ class DeviceWidgetnDisplay(DeviceWidgetUnreal):
 class DisplayConfig(object):
     ''' Encapsulates nDisplay config'''
 
-    validClassNames = ('DisplayClusterBlueprint','/Script/DisplayCluster.DisplayClusterBlueprint')
+    VALID_CLASS_NAMES = (
+        'DisplayClusterBlueprint',
+        '/Script/DisplayCluster.DisplayClusterBlueprint'
+    )
 
     def __init__(self):
         self.nodes = []
         self.uasset_path = ''
 
+
 class DevicenDisplay(DeviceUnreal):
 
     add_device_dialog = AddnDisplayDialog
@@ -394,7 +395,7 @@ class DevicenDisplay(DeviceUnreal):
             attr_name="render_sync_policy",
             nice_name="Render Sync Policy",
             value='Config',
-            possible_values=['Config','None','Ethernet','Nvidia'],
+            possible_values=['Config', 'None', 'Ethernet', 'Nvidia'],
             tool_tip=(
                 "Select which cluster synchronization policy to use. \n"
                 "- 'Config': Use the setting in the nDisplay config file \n"
@@ -503,7 +504,7 @@ class DevicenDisplay(DeviceUnreal):
             value=':0',
             tool_tip=(
                 'Local interface binding (-UDPMESSAGING_TRANSPORT_UNICAST) of '
-                'the form {ip}:{port}. If {ip} is omitted, the device IP '
+                'the form {address}:{port}. If {address} is omitted, the device '
                 'address is used.'),
         ),
         'udpmessaging_extra_static_endpoints': StringSetting(
@@ -532,8 +533,8 @@ class DevicenDisplay(DeviceUnreal):
     ndisplay_monitor_ui = None
     ndisplay_monitor = None
 
-    def __init__(self, name, ip_address, **kwargs):
-        super().__init__(name, ip_address, **kwargs)
+    def __init__(self, name, address, **kwargs):
+        super().__init__(name, address, **kwargs)
 
         self.settings = {
             'ue_command_line': StringSetting(
@@ -1027,7 +1028,7 @@ class DevicenDisplay(DeviceUnreal):
             # ConfigExport tag and parse it.
 
             for assetdata in aparser.aregdata:
-                if assetdata.ObjectClassName in DisplayConfig.validClassNames:
+                if assetdata.ObjectClassName in DisplayConfig.VALID_CLASS_NAMES:
                     return assetdata.tags['ConfigExport']
 
         raise ValueError('Invalid nDisplay config .uasset')
@@ -1090,8 +1091,8 @@ class DevicenDisplay(DeviceUnreal):
                     'node in the nDisplay configuration')
                 continue
 
-            # override ip address
-            node['host'] = device.ip_address
+            # override address
+            node['host'] = device.address
 
             # add to active nodes
             activenodes[device.name] = node
@@ -1139,7 +1140,7 @@ class DevicenDisplay(DeviceUnreal):
         return cls.parse_config_json_string(jsstr)
 
     @classmethod
-    def parse_config_json_string(cls, jsstr) -> Tuple[List, Optional[str]]:
+    def parse_config_json_string(cls, jsstr):
         '''
         Parses nDisplay config JSON string, returning (nodes, uasset_path).
         '''
@@ -1167,12 +1168,12 @@ class DevicenDisplay(DeviceUnreal):
             kwargs["window_resolution"] = (resx, resy)
             # Note the capital 'S'.
             kwargs["fullscreen"] = bool(cnode.get('fullScreen', False))
-        
+
             primary = True if primaryNode['id'] == name else False
 
             config.nodes.append({
                 "name": name,
-                "ip_address": cnode['host'],
+                "address": cnode['host'],
                 "primary": primary,
                 "port_ce": int(primaryNode['ports']['ClusterEventsJson']),
                 "kwargs": kwargs,
@@ -1187,7 +1188,7 @@ class DevicenDisplay(DeviceUnreal):
         return cls.parse_config_json_string(jsstr)
 
     @classmethod
-    def parse_config(cls, cfg_file) -> Tuple[List, Optional[str]]:
+    def parse_config(cls, cfg_file) -> DisplayConfig:
         '''
         Parses an nDisplay file and returns the nodes with the relevant
         information.
@@ -1288,7 +1289,7 @@ class DevicenDisplay(DeviceUnreal):
             with open(cfg_file, 'rb') as f:
                 cfg_content = f.read()
 
-        # Apply local overrides to configuration (e.g. ip addresses)
+        # Apply local overrides to configuration (e.g. addresses)
         try:
             cfg_content = self.__class__.apply_local_overrides_to_config(
                 cfg_content, cfg_ext).encode('utf-8')
@@ -1398,11 +1399,11 @@ class DevicenDisplay(DeviceUnreal):
         # primary.
         port = primary.nodeconfig['port_ce']
 
-        # Use the overridden ip address of this node
-        ip = primary.ip_address
+        # Use the overridden address of this node
+        address = primary.address
 
         sock = socket.socket(socket.AF_INET, socket.SOCK_STREAM)
-        sock.connect((ip, port))
+        sock.connect((address, port))
         sock.send(msg)
 
     @classmethod
diff --git a/Engine/Plugins/VirtualProduction/Switchboard/Source/Switchboard/switchboard/devices/shogun/plugin_shogun.py b/Engine/Plugins/VirtualProduction/Switchboard/Source/Switchboard/switchboard/devices/shogun/plugin_shogun.py
--- a/Engine/Plugins/VirtualProduction/Switchboard/Source/Switchboard/switchboard/devices/shogun/plugin_shogun.py
+++ b/Engine/Plugins/VirtualProduction/Switchboard/Source/Switchboard/switchboard/devices/shogun/plugin_shogun.py
@@ -34,9 +34,9 @@ def unresponsive_shogun(f):
 
 
 class DeviceShogun(Device):
-    def __init__(self, name, ip_address, **kwargs):
+    def __init__(self, name, address, **kwargs):
         self.setting_save_path = StringSetting("save_path", "Save Path", "")
-        super().__init__(name, ip_address, **kwargs)
+        super().__init__(name, address, **kwargs)
 
         self.trigger_start = True
         self.trigger_stop = True
@@ -67,7 +67,7 @@ class DeviceShogun(Device):
     def connect_listener(self):
         super().connect_listener()
 
-        self.client = vicon_core_api.Client(self.ip_address)
+        self.client = vicon_core_api.Client(self.address)
         self.capture_service = shogun_live_api.CaptureServices(self.client)
 
         if self.client.connected:
@@ -139,8 +139,8 @@ class DeviceShogun(Device):
 
 
 class DeviceWidgetShogun(DeviceWidget):
-    def __init__(self, name, device_hash, ip_address, icons, parent=None):
-        super().__init__(name, device_hash, ip_address, icons, parent=parent)
+    def __init__(self, name, device_hash, address, icons, parent=None):
+        super().__init__(name, device_hash, address, icons, parent=parent)
 
     def _add_control_buttons(self):
         super()._add_control_buttons()
diff --git a/Engine/Plugins/VirtualProduction/Switchboard/Source/Switchboard/switchboard/devices/sounddevices/plugin_sounddevices.py b/Engine/Plugins/VirtualProduction/Switchboard/Source/Switchboard/switchboard/devices/sounddevices/plugin_sounddevices.py
--- a/Engine/Plugins/VirtualProduction/Switchboard/Source/Switchboard/switchboard/devices/sounddevices/plugin_sounddevices.py
+++ b/Engine/Plugins/VirtualProduction/Switchboard/Source/Switchboard/switchboard/devices/sounddevices/plugin_sounddevices.py
@@ -12,8 +12,8 @@ import json
 
 
 class DeviceSoundDevices(Device):
-    def __init__(self, name, ip_address, **kwargs):
-        super().__init__(name, ip_address, **kwargs)
+    def __init__(self, name, address, **kwargs):
+        super().__init__(name, address, **kwargs)
 
         self.user_name = 'guest'
         self.password = 'guest'
@@ -53,7 +53,7 @@ class DeviceSoundDevices(Device):
         try:
             import requests
             from requests.auth import HTTPDigestAuth
-            command = f'http://{self.ip_address}/sounddevices/{command}'
+            command = f'http://{self.address}/sounddevices/{command}'
             response = requests.get(command, auth=HTTPDigestAuth(self.user_name, self.password), timeout=3)
             return response.content.decode('utf-8')
         except ModuleNotFoundError:
@@ -153,7 +153,7 @@ class DeviceSoundDevices(Device):
         for i in range(1,5):
             self._send_command(f'setsetting/RecordToDrive{i}=Ethernet%20File%20Transfer')
         '''
-        f'net use * \\\\{self.ip_address} /user:{self.user_name} {self.password}'
+        f'net use * \\\\{self.address} /user:{self.user_name} {self.password}'
         robocopy.exe "W:\wwwroot\MyProject" "\\192.168.0.1\Share\wwwroot\MyProject" *.* /E /XO /XD "App_Data/Search" "*.svn" /XF "sitefinity.log" "Thumbs.db" /NDL /NC /NP
         net use * /delete /yes
 
@@ -163,8 +163,8 @@ class DeviceSoundDevices(Device):
 
 
 class DeviceWidgetSoundDevices(DeviceWidget):
-    def __init__(self, name, device_type, device_hash, ip_address, parent=None):
-        super().__init__(name, device_type, device_hash, ip_address, parent=parent)
+    def __init__(self, name, device_type, device_hash, address, parent=None):
+        super().__init__(name, device_type, device_hash, address, parent=parent)
 
     def _add_control_buttons(self):
         super()._add_control_buttons()
diff --git a/Engine/Plugins/VirtualProduction/Switchboard/Source/Switchboard/switchboard/devices/unreal/plugin_unreal.py b/Engine/Plugins/VirtualProduction/Switchboard/Source/Switchboard/switchboard/devices/unreal/plugin_unreal.py
--- a/Engine/Plugins/VirtualProduction/Switchboard/Source/Switchboard/switchboard/devices/unreal/plugin_unreal.py
+++ b/Engine/Plugins/VirtualProduction/Switchboard/Source/Switchboard/switchboard/devices/unreal/plugin_unreal.py
@@ -7,10 +7,11 @@ import json
 from collections import OrderedDict
 from datetime import datetime
 from functools import wraps
-import ipaddress
+from ipaddress import IPv4Address
 import os
 import pathlib
 import re
+import socket
 import sys
 import threading
 from typing import Callable, List, Optional, Set
@@ -358,7 +359,7 @@ class DeviceUnreal(Device):
             value=':0',
             tool_tip=(
                 'Local interface binding (-UDPMESSAGING_TRANSPORT_UNICAST) of '
-                'the form {ip}:{port}. If {ip} is omitted, the device IP '
+                'the form {address}:{port}. If {address} is omitted, the device address '
                 'address is used.'),
         ),
         'udpmessaging_extra_static_endpoints': StringSetting(
@@ -376,7 +377,7 @@ class DeviceUnreal(Device):
             value='230.0.0.1:6666',
             tool_tip=(
                 'Multicast group and port (-UDPMESSAGING_TRANSPORT_MULTICAST) '
-                'in the {ip}:{port} endpoint format. The multicast group IP '
+                'in the {address}:{port} endpoint format. The multicast group address '
                 'must be in the range 224.0.0.0 to 239.255.255.255.'),
         ),
         'log_download_dir': DirectoryPathSetting(
@@ -434,24 +435,42 @@ class DeviceUnreal(Device):
     @classmethod
     def get_designated_local_builder(cls) -> Optional[DeviceUnreal]:
         '''
-        Returns first (by IP) local `DeviceUnreal` (or derived class) device.
+        Selects and returns a local `DeviceUnreal` (or derived class) device.
 
         This is the device tasked with building the multiuser server and
         listener executables.
         '''
-        ip_device_pairs = [
-            (ipaddress.ip_address(d.ip_address), d)
-            for d in cls.active_unreal_devices]
 
-        switchboard_ip_address = ipaddress.ip_address(SETTINGS.IP_ADDRESS.get_value())
-
-        def is_local(pair):
-            return pair[0].is_loopback or (pair[0] == switchboard_ip_address)
-
-        ip_device_pairs = filter(is_local, ip_device_pairs)
-        ip_device_pairs = sorted(ip_device_pairs, key=lambda pair: pair[0])
+        def ips_for_host(host: str) -> List[IPv4Address]:
+            try:
+                (primary_name, aliases, ipstrs) = socket.gethostbyname_ex(host)
+                return [IPv4Address(ipstr) for ipstr in ipstrs]
+            except Exception as exc:
+                LOGGER.error('get_designated_local_builder(): '
+                             f"gethostbyname_ex('{host}') exception", exc)
+                return []
+
+        # Use the set intersection of local IPs and device IPs to pick a device
+        address_to_device_map = dict[IPv4Address, DeviceUnreal]()
+        for device in cls.active_unreal_devices:
+            for ip in ips_for_host(device.address):
+                address_to_device_map[ip] = device
+
+        local_addr_set = set(ips_for_host(SETTINGS.ADDRESS.get_value()))
+        device_addr_set = set(address_to_device_map.keys())
+
+        local_device_addrs = local_addr_set.intersection(device_addr_set)
+        for local_device_addr in local_device_addrs:
+            local_device = address_to_device_map[local_device_addr]
+            if local_device.is_disconnected:
+                continue
+            else:
+                return local_device
 
-        return ip_device_pairs[0][1] if len(ip_device_pairs) > 0 else None
+        try:
+            return address_to_device_map[list(local_device_addrs)[0]]
+        except (IndexError, KeyError):
+            return None
 
     def is_designated_local_builder(self) -> bool:
         return self is DeviceUnreal.get_designated_local_builder()
@@ -476,11 +495,11 @@ class DeviceUnreal(Device):
         dlg.exec()
         DeviceUnreal._pending_notify_redeploy = False
 
-    def __init__(self, name, ip_address, **kwargs):
-        super().__init__(name, ip_address, **kwargs)
+    def __init__(self, name, address, **kwargs):
+        super().__init__(name, address, **kwargs)
 
         self.unreal_client = ListenerClient(
-            ip_address=self.ip_address,
+            address=self.address,
             port=DeviceUnreal.csettings['port'].get_value(self.name),
             buffer_size=DeviceUnreal.csettings['buffer_size'].get_value(
                 self.name)
@@ -531,8 +550,8 @@ class DeviceUnreal(Device):
             tool_tip="Whether to exclude this device from builds"
         )
 
-        self.setting_ip_address.signal_setting_changed.connect(
-            self.on_setting_ip_address_changed)
+        self.setting_address.signal_setting_changed.connect(
+            self.on_setting_address_changed)
         DeviceUnreal.csettings['port'].signal_setting_changed.connect(
             self.on_setting_port_changed)
         CONFIG.BUILD_ENGINE.signal_setting_changed.connect(
@@ -685,10 +704,10 @@ class DeviceUnreal(Device):
             CONFIG.listener_path())
 
         DeviceUnreal.rsync_server.register_client(
-            device, device.name, device.ip_address)
+            device, device.name, device.address)
         device.widget.signal_device_name_changed.connect(
             device.reregister_rsync_client)
-        device.setting_ip_address.signal_setting_changed.connect(
+        device.setting_address.signal_setting_changed.connect(
             device.reregister_rsync_client)
 
     @classmethod
@@ -703,7 +722,7 @@ class DeviceUnreal(Device):
         DeviceUnreal.rsync_server.unregister_client(device)
         device.widget.signal_device_name_changed.disconnect(
             device.reregister_rsync_client)
-        device.setting_ip_address.signal_setting_changed.disconnect(
+        device.setting_address.signal_setting_changed.disconnect(
             device.reregister_rsync_client)
 
         if len(DeviceUnreal.active_unreal_devices) == 0:
@@ -808,13 +827,13 @@ class DeviceUnreal(Device):
 
         device_widget.autojoin_mu.signal_device_widget_autojoin_mu.connect(self.on_autojoin_mu_ui_change)
 
-    def on_setting_ip_address_changed(self, _, new_address):
-        LOGGER.info(f"Updating IP address for ListenerClient to {new_address}")
-        self.unreal_client.ip_address = new_address
+    def on_setting_address_changed(self, _, new_address):
+        LOGGER.info(f"Updating address for ListenerClient to {new_address}")
+        self.unreal_client.address = new_address
 
     def on_setting_port_changed(self, _, new_port):
         if not DeviceUnreal.csettings['port'].is_overridden(self.name):
-            LOGGER.info(f"Updating Port for ListenerClient to {new_port}")
+            LOGGER.info(f"Updating port for ListenerClient to {new_port}")
             self.unreal_client.port = new_port
 
     def on_build_engine_changed(self, _, build_engine):
@@ -1114,6 +1133,8 @@ class DeviceUnreal(Device):
         # Build dependency chain
         puuid_dependency = sync_puuid
 
+        # TODO: Corner case if multiple local devices, and we build on a single
+        # local device other than chosen. Guarantee on any local single build?
         if (self.is_designated_local_builder() and
                 CONFIG.BUILD_ENGINE.get_value()):
             # Build multi-user server
@@ -1314,7 +1335,7 @@ class DeviceUnreal(Device):
         setting_val = CONFIG.MUSERVER_ENDPOINT.get_value().strip()
         if setting_val:
             return sb_utils.expand_endpoint(setting_val,
-                                            SETTINGS.IP_ADDRESS.get_value().strip())
+                                            SETTINGS.ADDRESS.get_value().strip())
         else:
             return ''
 
@@ -1322,7 +1343,7 @@ class DeviceUnreal(Device):
     def udpmessaging_unicast_endpoint(self) -> str:
         setting_val = self.udpmessaging_unicast_endpoint_setting.strip()
         if setting_val:
-            return sb_utils.expand_endpoint(setting_val, self.ip_address)
+            return sb_utils.expand_endpoint(setting_val, self.address)
         else:
             return ''
 
@@ -1536,7 +1557,7 @@ class DeviceUnreal(Device):
         if self.status == DeviceStatus.OPEN:
             self.send_osc_message(
                 osc.OSC_ADD_SEND_TARGET,
-                [SETTINGS.IP_ADDRESS.get_value(), CONFIG.OSC_SERVER_PORT.get_value()])
+                [SETTINGS.ADDRESS.get_value(), CONFIG.OSC_SERVER_PORT.get_value()])
         else:
             self.osc_connection_timer.stop()
 
@@ -1973,7 +1994,7 @@ class DeviceUnreal(Device):
     def reregister_rsync_client(self):
         DeviceUnreal.rsync_server.unregister_client(self)
         DeviceUnreal.rsync_server.register_client(
-            self, self.name, self.ip_address)
+            self, self.name, self.address)
 
     def backup_file(self, filename):
         ''' Rotate existing filename to a timestamped backup, a la Unreal. '''
@@ -2014,7 +2035,7 @@ class DeviceUnreal(Device):
             DeviceUnreal.rsync_server.make_cygdrive_path(remote_path)
 
         dest_endpoint = \
-            f'{SETTINGS.IP_ADDRESS.get_value()}:{DeviceUnreal.rsync_server.port}'
+            f'{SETTINGS.ADDRESS.get_value()}:{DeviceUnreal.rsync_server.port}'
         dest_module = DeviceUnreal.rsync_server.INCOMING_LOGS_MODULE
         dest_path = f'rsync://{dest_endpoint}/{dest_module}/'
 
@@ -2182,7 +2203,7 @@ class DeviceWidgetUnreal(DeviceWidget):
     signal_open_last_trace = QtCore.Signal(object)
     signal_copy_last_launch_command = QtCore.Signal(object)
 
-    def __init__(self, name, device_hash, ip_address, icons, parent=None):
+    def __init__(self, name, device_hash, address, icons, parent=None):
         self._autojoin_visible = True
         self._is_engine_synched = True
         self._is_project_synched = True
@@ -2190,7 +2211,7 @@ class DeviceWidgetUnreal(DeviceWidget):
         self._exclude_from_build = False
         self._desired_build_button_tooltip = "Build changelist"
 
-        super().__init__(name, device_hash, ip_address, icons, parent=parent)
+        super().__init__(name, device_hash, address, icons, parent=parent)
 
         CONFIG.P4_ENABLED.signal_setting_changed.connect(
             lambda _, enabled: self.sync_button.setVisible(enabled))
diff --git a/Engine/Plugins/VirtualProduction/Switchboard/Source/Switchboard/switchboard/devices/unreal/redeploy_dialog.py b/Engine/Plugins/VirtualProduction/Switchboard/Source/Switchboard/switchboard/devices/unreal/redeploy_dialog.py
--- a/Engine/Plugins/VirtualProduction/Switchboard/Source/Switchboard/switchboard/devices/unreal/redeploy_dialog.py
+++ b/Engine/Plugins/VirtualProduction/Switchboard/Source/Switchboard/switchboard/devices/unreal/redeploy_dialog.py
@@ -1,5 +1,7 @@
 # Copyright Epic Games, Inc. All Rights Reserved.
 
+from __future__ import annotations
+
 import base64
 import hashlib
 import os
@@ -9,7 +11,6 @@ from typing import Dict, Optional
 from PySide2 import QtCore, QtWidgets
 
 from switchboard import message_protocol
-from switchboard.config import CONFIG
 from switchboard.listener_client import ListenerClient
 from switchboard.switchboard_logging import LOGGER
 from . import version_helpers
@@ -21,19 +22,19 @@ class RedeployListenerEndpoint(QtCore.QObject):
 
     def __init__(
         self,
-        parent: 'RedeployListenerDialog',
-        ip_address: str,
+        parent: RedeployListenerDialog,
+        address: str,
         port: int
     ):
         super().__init__(parent)
-        
+
         self.dlg_parent = parent
         self.dlg_parent.signal_listener_changed.connect(self.refresh_ui)
 
-        self.ip_address = ip_address
+        self.address = address
         self.port = port
 
-        self.client = ListenerClient(self.ip_address, self.port)
+        self.client = ListenerClient(self.address, self.port)
         self.client.listener_qt_handler.listener_connecting.connect(
             lambda _: self.version_label.setText('Connecting...'))
         self.client.listener_qt_handler.listener_connection_failed.connect(
@@ -54,7 +55,7 @@ class RedeployListenerEndpoint(QtCore.QObject):
 
         self.signal_refresh_ui.connect(self.refresh_ui)
 
-        self.endpoint_label = QtWidgets.QLabel(f"{self.ip_address}:{self.port}")
+        self.endpoint_label = QtWidgets.QLabel(f"{self.address}:{self.port}")
         self.endpoint_label.setObjectName('endpoint_label')
 
         self.devices_label = QtWidgets.QLabel('')
@@ -190,9 +191,9 @@ class RedeployListenerDialog(QtWidgets.QDialog):
 
         # Remote listener models and rows
         for device in devices:
-            ep_addr = (device.unreal_client.ip_address, device.unreal_client.port)
+            ep_addr = (device.unreal_client.address, device.unreal_client.port)
             if ep_addr not in self.endpoints:
-                endpoint = RedeployListenerEndpoint(self, device.unreal_client.ip_address, device.unreal_client.port)
+                endpoint = RedeployListenerEndpoint(self, device.unreal_client.address, device.unreal_client.port)
                 endpoint.signal_result.connect(self.on_endpoint_result)
                 endpoint.signal_refresh_ui.connect(self.refresh_ui)
                 endpoint.client.connect()
@@ -283,4 +284,4 @@ class RedeployListenerDialog(QtWidgets.QDialog):
         sender = self.sender()
         assert isinstance(sender, RedeployListenerEndpoint)
         logfn = LOGGER.info if success else LOGGER.error
-        logfn(f"Redeploy endpoint {sender.ip_address}:{sender.port}: {details}", exc_info=exc_info)
+        logfn(f"Redeploy endpoint {sender.address}:{sender.port}: {details}", exc_info=exc_info)
diff --git a/Engine/Plugins/VirtualProduction/Switchboard/Source/Switchboard/switchboard/listener_client.py b/Engine/Plugins/VirtualProduction/Switchboard/Source/Switchboard/switchboard/listener_client.py
--- a/Engine/Plugins/VirtualProduction/Switchboard/Source/Switchboard/switchboard/listener_client.py
+++ b/Engine/Plugins/VirtualProduction/Switchboard/Source/Switchboard/switchboard/listener_client.py
@@ -1,13 +1,18 @@
 # Copyright Epic Games, Inc. All Rights Reserved.
-from . import message_protocol
-from .switchboard_logging import LOGGER
-
-from PySide2 import QtCore
-
-import datetime, select, socket, uuid, traceback, typing
 
 from collections import deque
+import datetime
+import select
+import socket
 from threading import Thread
+import traceback
+from typing import Callable, Dict, Optional
+import uuid
+
+from PySide2 import QtCore
+
+from . import message_protocol
+from .switchboard_logging import LOGGER
 
 
 class ListenerQtHandler(QtCore.QObject):
@@ -38,8 +43,8 @@ class ListenerClient(object):
     HITCH_THRESHOLD_SEC = (KEEPALIVE_INTERVAL_SEC + SELECT_TIMEOUT_SEC
                            + EXTRA_HITCH_TOLERANCE_SEC)
 
-    def __init__(self, ip_address, port, buffer_size=1024):
-        self.ip_address = ip_address
+    def __init__(self, address, port, buffer_size=1024):
+        self.address = address
         self.port = port
         self.buffer_size = buffer_size
 
@@ -51,7 +56,7 @@ class ListenerClient(object):
         self.socket = None
         self.handle_connection_thread = None
 
-        #TODO: Consider converting these delegates to Signals and sending dict.
+        # TODO: Consider converting these delegates to Signals and sending dict
 
         self.disconnect_delegate = None
 
@@ -70,17 +75,17 @@ class ListenerClient(object):
         self.receive_file_completed_delegate = None
         self.receive_file_failed_delegate = None
 
-        self.delegates: typing.Dict[ str, typing.Optional[ typing.Callable[[typing.Dict], None] ] ] = {
-            "state" : None,
-            "get sync status" : None,
+        self.delegates: Dict[str, Optional[Callable[[Dict], None]]] = {
+            "state": None,
+            "get sync status": None,
         }
 
         self.last_activity = datetime.datetime.now()
 
     @property
     def server_address(self):
-        if self.ip_address:
-            return (self.ip_address, self.port)
+        if self.address:
+            return (self.address, self.port)
         return None
 
     @property
@@ -95,7 +100,7 @@ class ListenerClient(object):
             return False
         return False
 
-    def connect(self, ip_address=None):
+    def connect(self, address=None):
         '''
         Initiates a connection.
 
@@ -106,10 +111,10 @@ class ListenerClient(object):
         '''
         self.disconnect()
 
-        if ip_address:
-            self.ip_address = ip_address
-        elif not self.ip_address:
-            LOGGER.debug('No ip_address has been set. Cannot connect')
+        if address:
+            self.address = address
+        elif not self.address:
+            LOGGER.debug('No address has been set. Cannot connect')
             self.listener_qt_handler.listener_connection_failed.emit(self)
             return False
 
@@ -124,7 +129,7 @@ class ListenerClient(object):
 
     def _establish_connection(self):
         try:
-            LOGGER.info(f"Connecting to {self.ip_address}:{self.port}")
+            LOGGER.info(f"Connecting to {self.address}:{self.port}")
 
             self.socket = socket.socket(socket.AF_INET, socket.SOCK_STREAM)
             self.socket.connect(self.server_address)
@@ -134,7 +139,7 @@ class ListenerClient(object):
             self.handle_connection_thread.start()
 
         except OSError:
-            LOGGER.error(f"Socket error: {self.ip_address}:{self.port}")
+            LOGGER.error(f"Socket error: {self.address}:{self.port}")
             self.socket = None
             self.listener_qt_handler.listener_connection_failed.emit(self)
             return False
@@ -157,7 +162,7 @@ class ListenerClient(object):
         if delta_sec > self.HITCH_THRESHOLD_SEC:
             LOGGER.warning(
                 f'Hitch detected; {delta_sec:.1f} seconds since last keepalive'
-                f' to {self.ip_address}')
+                f' to {self.address}')
 
     def handle_connection(self):
         buffer = []
@@ -290,9 +295,9 @@ class ListenerClient(object):
 
     def send_message(self, message_bytes):
         if self.is_connected:
-            LOGGER.message(f'Message: Sending ({self.ip_address}): {message_bytes}')
+            LOGGER.message(f'Message: Sending ({self.address}): {message_bytes}')
             self.message_queue.appendleft(message_bytes)
         else:
-            LOGGER.error(f'Message: Failed to send ({self.ip_address}): {message_bytes}. No socket connected')
+            LOGGER.error(f'Message: Failed to send ({self.address}): {message_bytes}. No socket connected')
             if self.disconnect_delegate:
                 self.disconnect_delegate(unexpected=True, exception=None)
diff --git a/Engine/Plugins/VirtualProduction/Switchboard/Source/Switchboard/switchboard/settings_dialog.py b/Engine/Plugins/VirtualProduction/Switchboard/Source/Switchboard/switchboard/settings_dialog.py
--- a/Engine/Plugins/VirtualProduction/Switchboard/Source/Switchboard/switchboard/settings_dialog.py
+++ b/Engine/Plugins/VirtualProduction/Switchboard/Source/Switchboard/switchboard/settings_dialog.py
@@ -204,11 +204,11 @@ class SettingsDialog(QtCore.QObject):
         self.ui.switchboard_settings_group = CollapsibleGroupBox()
         self.ui.switchboard_settings_group.setTitle("Switchboard")
         layout = QtWidgets.QFormLayout(self.ui.switchboard_settings_group)
-        
-        settings.IP_ADDRESS.create_ui(form_layout=layout)
+
+        settings.ADDRESS.create_ui(form_layout=layout)
         settings.TRANSPORT_PATH.create_ui(form_layout=layout)
         config.LISTENER_EXE.create_ui(form_layout=layout)
-        
+
         return self.ui.switchboard_settings_group
         
     def _create_project_settings(self, config: Config):
diff --git a/Engine/Plugins/VirtualProduction/Switchboard/Source/Switchboard/switchboard/switchboard_application.py b/Engine/Plugins/VirtualProduction/Switchboard/Source/Switchboard/switchboard/switchboard_application.py
--- a/Engine/Plugins/VirtualProduction/Switchboard/Source/Switchboard/switchboard/switchboard_application.py
+++ b/Engine/Plugins/VirtualProduction/Switchboard/Source/Switchboard/switchboard/switchboard_application.py
@@ -28,30 +28,30 @@ class OscServer:
         self.internal_client = None
         self.server_port = None
 
-    def ip_address(self):
+    def address(self):
         if not self.server:
             return None
 
         return self.server.server_address
 
-    def launch(self, ip_address, port):
-        # TODO: Allow relaunch of OSC server when ip_address variable changes
+    def launch(self, address, port):
+        # TODO: Allow relaunch of OSC server when address variable changes
         try:
             self.server = pythonosc.osc_server.ThreadingOSCUDPServer(
-                (ip_address, port), self.dispatcher)
+                (address, port), self.dispatcher)
         except OSError as e:
             if e.errno == 10048:
                 LOGGER.error(
                     'OSC Server: Another OSC server is currently using '
-                    f'{ip_address}:{port}. Please kill and relaunch')
+                    f'{address}:{port}. Please kill and relaunch')
             elif e.errno == 10049:
-                LOGGER.error(f"OSC Server: Couldn't bind {ip_address}:{port}. "
+                LOGGER.error(f"OSC Server: Couldn't bind {address}:{port}. "
                              'Please check address and relaunch')
 
             self.server = None
             return False
 
-        # Set the server ip and port
+        # Set the server address and port
         self.server_port = port
 
         LOGGER.success(
@@ -83,8 +83,6 @@ class OscServer:
         LOGGER.warning(f'Received unhandled OSC message: {command} {args}.')
 
 
-MultiUserServerInstance = None
-
 class MultiUserApplication:
     def __init__(self):
         self.lock = threading.Lock()
@@ -106,11 +104,11 @@ class MultiUserApplication:
         return ''
 
     def get_mu_server_endpoint_arg(self):
-        setting_val = CONFIG.MUSERVER_ENDPOINT.get_value().strip()
-        if setting_val:
-            converted_ip = sb_utils.expand_endpoint(setting_val,
-                                                    SETTINGS.IP_ADDRESS.get_value().strip())
-            return f'-UDPMESSAGING_TRANSPORT_UNICAST="{converted_ip}"'
+        endpoint_setting = CONFIG.MUSERVER_ENDPOINT.get_value().strip()
+        if endpoint_setting:
+            addr_setting = SETTINGS.ADDRESS.get_value().strip()
+            endpoint = sb_utils.expand_endpoint(endpoint_setting, addr_setting)
+            return f'-UDPMESSAGING_TRANSPORT_UNICAST="{endpoint}"'
         return ''
 
     def poll_process(self):
@@ -181,12 +179,17 @@ class MultiUserApplication:
 
         return False
 
+
+MultiUserServerInstance: Optional[MultiUserApplication] = None
+
+
 def get_multi_user_server_instance():
     global MultiUserServerInstance
     if not MultiUserServerInstance:
         MultiUserServerInstance = MultiUserApplication()
     return MultiUserServerInstance
 
+
 class RsyncServer:
     DEFAULT_PORT = 8730
     INCOMING_LOGS_MODULE = 'device_logs'
@@ -282,7 +285,7 @@ hosts allow = {allowed_addrs}
             client.address for client in self.allowed_clients.values()}
 
         allowed_addrs.add(self.address)
-        allowed_addrs.add(SETTINGS.IP_ADDRESS.get_value())
+        allowed_addrs.add(SETTINGS.ADDRESS.get_value())
         allowed_addrs.discard('0.0.0.0')
 
         config = self.CONFIG_TEMPLATE.format(
diff --git a/Engine/Plugins/VirtualProduction/Switchboard/Source/Switchboard/switchboard/switchboard_dialog.py b/Engine/Plugins/VirtualProduction/Switchboard/Source/Switchboard/switchboard/switchboard_dialog.py
--- a/Engine/Plugins/VirtualProduction/Switchboard/Source/Switchboard/switchboard/switchboard_dialog.py
+++ b/Engine/Plugins/VirtualProduction/Switchboard/Source/Switchboard/switchboard/switchboard_dialog.py
@@ -310,7 +310,7 @@ class SwitchboardDialog(QtCore.QObject):
 
         # Start the OSC server
         self.osc_server = switchboard_application.OscServer()
-        self.osc_server.launch(SETTINGS.IP_ADDRESS.get_value(), CONFIG.OSC_SERVER_PORT.get_value())
+        self.osc_server.launch(SETTINGS.ADDRESS.get_value(), CONFIG.OSC_SERVER_PORT.get_value())
 
         # Register with OSC server
         self.osc_server.dispatcher_map(osc.TAKE, self.osc_take)
@@ -324,8 +324,8 @@ class SwitchboardDialog(QtCore.QObject):
         self.osc_server.dispatcher_map(osc.RECORD_CANCEL_CONFIRM, self.osc_record_cancel_confirm)
         self.osc_server.dispatcher_map(osc.UE4_LAUNCH_CONFIRM, self.osc_ue4_launch_confirm)
         self.osc_server.dispatcher.map(osc.OSC_ADD_SEND_TARGET_CONFIRM, self.osc_add_send_target_confirm, 1, needs_reply_address=True)
-        self.osc_server.dispatcher.map(osc.ARSESSION_START_CONFIRM, self.osc_arsession_start_confrim, 1, needs_reply_address=True)
-        self.osc_server.dispatcher.map(osc.ARSESSION_STOP_CONFIRM, self.osc_arsession_stop_confrim, 1, needs_reply_address=True)
+        self.osc_server.dispatcher.map(osc.ARSESSION_START_CONFIRM, self.osc_arsession_start_confirm, 1, needs_reply_address=True)
+        self.osc_server.dispatcher.map(osc.ARSESSION_STOP_CONFIRM, self.osc_arsession_stop_confirm, 1, needs_reply_address=True)
         self.osc_server.dispatcher_map(osc.BATTERY, self.osc_battery)
         self.osc_server.dispatcher_map(osc.DATA, self.osc_data)
 
@@ -425,37 +425,36 @@ class SwitchboardDialog(QtCore.QObject):
         #self.transport_queue_resume()
         
         self.update_current_config_text()
-        self.update_current_ip_address_text()
-        SETTINGS.IP_ADDRESS.signal_setting_changed.connect(
-            lambda: self.update_current_ip_address_text()
+        self.update_current_address_text()
+        SETTINGS.ADDRESS.signal_setting_changed.connect(
+            lambda: self.update_current_address_text()
         )
-        self.window.current_ip_value.editingFinished.connect(self._try_change_ip_address)
+        self.window.current_address_value.editingFinished.connect(self._try_change_address)
 
         self.refresh_window_title()
 
         self.script_manager.on_postinit(self)
         
-    def _try_change_ip_address(self):
-        def is_valid_ip_format(address):
-            ip_port_split = address.split(":")
-            if len(ip_port_split) != 1:
-                return False
-            
-            ip_elements = ip_port_split[0].split(".")
-            if len(ip_elements) != 4:
-                return False
-            
-            return all(element.isdigit() for element in ip_elements)
+    def _try_change_address(self):
+        new_value = self.window.current_address_value.text()
+        old_address = SETTINGS.ADDRESS.get_value()
         
-        new_value = self.window.current_ip_value.text()
-        if is_valid_ip_format(new_value):
-            SETTINGS.IP_ADDRESS.update_value(new_value)
-            SETTINGS.save()
-            # Print to make it clear to the user that the IP was updated
-            LOGGER.info(f"Updated IP to {new_value}")
+        if not SETTINGS.ADDRESS.possible_values.__contains__(new_value):
+            SETTINGS.ADDRESS.possible_values.append(new_value)
+            
+        SETTINGS.ADDRESS.update_value(new_value)
+        SETTINGS.save()
+        # Check if osc connection binds before commiting changes
+        if self.osc_restart_check(old_address):
+            # Print to make it clear to the user that the Address was updated
+            if old_address != new_value:
+                LOGGER.info(f"Updated address to {new_value}")
         else:
-            self.window.current_ip_value.setText(SETTINGS.IP_ADDRESS.get_value())
-            LOGGER.warning(f"{new_value} must be in the format a.b.c.d (without port)")
+            LOGGER.warning(f"Reverting to the previous address ({old_address})")
+            SETTINGS.ADDRESS.update_value(old_address)
+            SETTINGS.save()
+            self.osc_server.launch(SETTINGS.ADDRESS.get_value(), CONFIG.OSC_SERVER_PORT.get_value())
+            
 
     def _poll_muserver_status(self):
         '''
@@ -840,8 +839,8 @@ class SwitchboardDialog(QtCore.QObject):
         else:
             self.window.current_config_file_value.setText("No config loaded")
     
-    def update_current_ip_address_text(self):
-        self.window.current_ip_value.setText(SETTINGS.IP_ADDRESS.get_value())
+    def update_current_address_text(self):
+        self.window.current_address_value.setText(SETTINGS.ADDRESS.get_value())
 
     def create_new_config(self, file_path: Union[str, Path], uproject, engine_dir, p4_settings):
         ''' Creates a new config file
@@ -986,7 +985,7 @@ class SwitchboardDialog(QtCore.QObject):
         
         settings_dialog.select_all_tab()
 
-        old_ip_address = SETTINGS.IP_ADDRESS.get_value()
+        old_address = SETTINGS.ADDRESS.get_value()
 
         # avoid saving the config all the time while in the settings dialog
         CONFIG.push_saving_allowed(False)
@@ -1002,15 +1001,20 @@ class SwitchboardDialog(QtCore.QObject):
             CONFIG.replace(new_config_path)
             SETTINGS.CONFIG = new_config_path
             SETTINGS.save()
-
-        if old_ip_address != SETTINGS.IP_ADDRESS.get_value() or SETTINGS.TRANSPORT_PATH.get_value():
+        if old_address != SETTINGS.ADDRESS.get_value() or SETTINGS.TRANSPORT_PATH.get_value():
             SETTINGS.save()
-        if old_ip_address != SETTINGS.IP_ADDRESS.get_value():
-            self.osc_server.close()
-            self.osc_server.launch(SETTINGS.IP_ADDRESS.get_value(), CONFIG.OSC_SERVER_PORT.get_value())
+
+        self.osc_restart_check(old_address)
 
         CONFIG.save()
 
+    def osc_restart_check(self, old_address):
+        if old_address != SETTINGS.ADDRESS.get_value():
+            self.osc_server.close()
+            return self.osc_server.launch(SETTINGS.ADDRESS.get_value(), CONFIG.OSC_SERVER_PORT.get_value())
+        else:
+            return True
+
     def sync_all_button_clicked(self):
         if not CONFIG.P4_ENABLED.get_value():
             return
@@ -1259,9 +1263,9 @@ class SwitchboardDialog(QtCore.QObject):
     def slate(self, value):
         self._set_slate(value, reset_take=False)
 
-    def _set_slate(self, value, exclude_ip_addresses=[], reset_take=True):
+    def _set_slate(self, value, exclude_addresses=[], reset_take=True):
         """
-        Internal setter that allows exclusion of ip addresses
+        Internal setter that allows exclusion of addresses
         """
         # Protect against blank slates
         if value == '':
@@ -1282,15 +1286,15 @@ class SwitchboardDialog(QtCore.QObject):
         if self.window.slate_line_edit.text() != self._slate:
             self.window.slate_line_edit.setText(self._slate)
 
-        thread = threading.Thread(target=self._set_slate_all_devices, args=[value], kwargs={'exclude_ip_addresses':exclude_ip_addresses})
+        thread = threading.Thread(target=self._set_slate_all_devices, args=[value], kwargs={'exclude_addresses':exclude_addresses})
         thread.start()
 
-    def _set_slate_all_devices(self, value, exclude_ip_addresses=[]):
+    def _set_slate_all_devices(self, value, exclude_addresses=[]):
         """
         Tell all devices the new slate
         """
         for device in self.device_manager.devices():
-            if device.ip_address in exclude_ip_addresses:
+            if device.address in exclude_addresses:
                 continue
 
             # Do not send updates to disconnected devices
@@ -1307,9 +1311,9 @@ class SwitchboardDialog(QtCore.QObject):
     def take(self, value):
         self._set_take(value)
 
-    def _set_take(self, value, exclude_ip_addresses=[]):
+    def _set_take(self, value, exclude_addresses=[]):
         """
-        Internal setter that allows exclusion of ip addresses
+        Internal setter that allows exclusion of addresses
         """
         requested_take = value
 
@@ -1324,7 +1328,7 @@ class SwitchboardDialog(QtCore.QObject):
         if requested_take != value:
             LOGGER.warning(f'Slate: "{self._slate}" Take: "{value}" have already been used. Auto incremented up to take: "{requested_take}"')
             # Clear the exclude list since Switchboard changed the incoming value
-            exclude_ip_addresses = []
+            exclude_addresses = []
         
 
         self._take = requested_take
@@ -1334,10 +1338,10 @@ class SwitchboardDialog(QtCore.QObject):
         if self.window.take_spin_box.value() != self._take:
             self.window.take_spin_box.setValue(self._take)
 
-        thread = threading.Thread(target=self._set_take_all_devices, args=[value], kwargs={'exclude_ip_addresses':exclude_ip_addresses})
+        thread = threading.Thread(target=self._set_take_all_devices, args=[value], kwargs={'exclude_addresses':exclude_addresses})
         thread.start()
 
-    def _set_take_all_devices(self, value, exclude_ip_addresses=[]):
+    def _set_take_all_devices(self, value, exclude_addresses=[]):
         """
         Tell all devices the new take
         """
@@ -1347,7 +1351,7 @@ class SwitchboardDialog(QtCore.QObject):
             if device.is_disconnected:
                 continue
 
-            if device.ip_address in exclude_ip_addresses:
+            if device.address in exclude_addresses:
                 continue
 
             device.set_take(self._take)
@@ -1540,9 +1544,9 @@ class SwitchboardDialog(QtCore.QObject):
         )
         
         cl = device.project_changelist
-        ip = device.ip_address
+        address = device.address
         for device in self.device_manager.devices():
-            if device.ip_address == ip and device.project_changelist and device.project_changelist != cl:
+            if device.address == address and device.project_changelist and device.project_changelist != cl:
                 device.project_changelist = cl
 
     @QtCore.Slot(object)
@@ -1555,9 +1559,9 @@ class SwitchboardDialog(QtCore.QObject):
         )
 
         cl = device.engine_changelist
-        ip = device.ip_address
+        address = device.address
         for device in self.device_manager.devices():
-            if device.ip_address == ip and device.engine_changelist and device.engine_changelist != cl:
+            if device.address == address and device.engine_changelist and device.engine_changelist != cl:
                 device.engine_changelist = cl
 
     @QtCore.Slot(object)
@@ -1885,24 +1889,24 @@ class SwitchboardDialog(QtCore.QObject):
             self.p4_refresh_engine_cl()
             self.p4_refresh_project_cl()
 
-    def osc_take(self, ip_address, command, value):
-        device = self._device_from_ip_address(ip_address, command, value=value)
+    def osc_take(self, address, command, value):
+        device = self._device_from_address(address, command, value=value)
 
         if not device:
             return
 
-        self._set_take(value, exclude_ip_addresses=[device.ip_address])
+        self._set_take(value, exclude_addresses=[device.address])
 
     # OSC: Set Slate
-    def osc_slate(self, ip_address, command, value):
-        device = self._device_from_ip_address(ip_address, command, value=value)
+    def osc_slate(self, address, command, value):
+        device = self._device_from_address(address, command, value=value)
         if not device:
             return
 
-        self._set_slate(value, exclude_ip_addresses=[device.ip_address], reset_take=False)
+        self._set_slate(value, exclude_addresses=[device.address], reset_take=False)
 
     # OSC: Set Description UPDATE THIS TO MAKE IT WORK
-    def osc_slate_description(self, ip_address, command, value):
+    def osc_slate_description(self, address, command, value):
         self.description = value
 
     def record_button_released(self):
@@ -1917,11 +1921,11 @@ class SwitchboardDialog(QtCore.QObject):
             LOGGER.debug('Record start button pressed')
             self._record_start(self.slate, self.take, self.description)
 
-    def _device_from_ip_address(self, ip_address, command, value=''):
-        device = self.device_manager.device_with_ip_address(ip_address[0])
+    def _device_from_address(self, address, command, value=''):
+        device = self.device_manager.device_with_address(address[0])
 
         if not device:
-            LOGGER.warning(f'{ip_address} is not registered with a device in Switchboard')
+            LOGGER.warning(f'{address} is not registered with a device in Switchboard')
             return None
 
         # Do not receive osc from disconnected devices
@@ -1929,15 +1933,15 @@ class SwitchboardDialog(QtCore.QObject):
             LOGGER.warning(f'{device.name}: is sending OSC commands but is not connected. Ignoring "{command} {value}"')
             return None
 
-        LOGGER.osc(f'OSC Server: Received "{command} {value}" from {device.name} ({device.ip_address})')
+        LOGGER.osc(f'OSC Server: Received "{command} {value}" from {device.name} ({device.address})')
         return device
 
     # OSC: Start a recording
-    def osc_record_start(self, ip_address, command, slate, take, description):
+    def osc_record_start(self, address, command, slate, take, description):
         '''
         OSC message Recieved /RecordStart
         '''
-        device = self._device_from_ip_address(ip_address, command, value=[slate, take, description])
+        device = self._device_from_address(address, command, value=[slate, take, description])
         if not device:
             return
 
@@ -1951,9 +1955,9 @@ class SwitchboardDialog(QtCore.QObject):
         self.take = take
         self.description = description
 
-        self._record_start(self.slate, self.take, self.description, exclude_ip_address=device.ip_address)
+        self._record_start(self.slate, self.take, self.description, exclude_address=device.address)
 
-    def _record_start(self, slate, take, description, exclude_ip_address=None):
+    def _record_start(self, slate, take, description, exclude_address=None):
         LOGGER.success(f'Record Start: "{self.slate}" {self.take}')
 
         # Update the UI button
@@ -1987,7 +1991,7 @@ class SwitchboardDialog(QtCore.QObject):
         devices = self.device_manager.devices()
         for device in devices:
             # Do not send a start record message to whichever device sent it
-            if exclude_ip_address and exclude_ip_address == device.ip_address:
+            if exclude_address and exclude_address == device.address:
                 continue
 
             # Do not send updates to disconnected devices
@@ -1997,7 +2001,7 @@ class SwitchboardDialog(QtCore.QObject):
             LOGGER.debug(f'Record Start {device.name}')
             device.record_start(slate, take, description)
 
-    def _record_stop(self, exclude_ip_address=None):
+    def _record_stop(self, exclude_address=None):
         LOGGER.success(f'Record Stop: "{self.slate}" {self.take}')
 
         pixmap = QtGui.QPixmap(":/icons/images/record_stop.png")
@@ -2023,7 +2027,7 @@ class SwitchboardDialog(QtCore.QObject):
         # Sends the message to all recording devices
         for device in devices:
             # Do not send a start record message to whichever device sent it
-            if exclude_ip_address and exclude_ip_address == device.ip_address:
+            if exclude_address and exclude_address == device.address:
                 continue
 
             # Do not send updates to disconnected devices
@@ -2032,29 +2036,29 @@ class SwitchboardDialog(QtCore.QObject):
 
             device.record_stop()
 
-    def _record_cancel(self, exclude_ip_address=None):
-        self._record_stop(exclude_ip_address=exclude_ip_address)
+    def _record_cancel(self, exclude_address=None):
+        self._record_stop(exclude_address=exclude_address)
 
-        # Incriment Take
+        # Increment Take
         #new_recording = self.recording_manager.latest_recording
         #self.take = new_recording.take + 1
 
-    def osc_record_start_confirm(self, ip_address, command, timecode):
-        device = self._device_from_ip_address(ip_address, command, value=timecode)
+    def osc_record_start_confirm(self, address, command, timecode):
+        device = self._device_from_address(address, command, value=timecode)
         if not device:
             return
 
         device.record_start_confirm(timecode)
 
-    def osc_record_stop(self, ip_address, command):
-        device = self._device_from_ip_address(ip_address, command)
+    def osc_record_stop(self, address, command):
+        device = self._device_from_address(address, command)
         if not device:
             return
 
-        self._record_stop(exclude_ip_address=device.ip_address)
+        self._record_stop(exclude_address=device.address)
 
-    def osc_record_stop_confirm(self, ip_address, command, timecode, *paths):
-        device = self._device_from_ip_address(ip_address, command, value=timecode)
+    def osc_record_stop_confirm(self, address, command, timecode, *paths):
+        device = self._device_from_address(address, command, value=timecode)
         if not device:
             return
 
@@ -2062,27 +2066,27 @@ class SwitchboardDialog(QtCore.QObject):
             paths = None
         device.record_stop_confirm(timecode, paths=paths)
 
-    def osc_record_cancel(self, ip_address, command):
+    def osc_record_cancel(self, address, command):
         """
         This is called when record has been pressed and stopped before the countdown in take recorder
         has finished
         """
-        device = self._device_from_ip_address(ip_address, command)
+        device = self._device_from_address(address, command)
         if not device:
             return
 
-        self._record_cancel(exclude_ip_address=device.ip_address)
+        self._record_cancel(exclude_address=device.address)
 
-    def osc_record_cancel_confirm(self, ip_address, command, timecode):
+    def osc_record_cancel_confirm(self, address, command, timecode):
         pass
-        #device = self._device_from_ip_address(ip_address, command, value=timecode)
+        #device = self._device_from_address(address, command, value=timecode)
         #if not device:
         #    return
 
         #self.record_cancel_confirm(device, timecode)
 
-    def osc_ue4_launch_confirm(self, ip_address, command):
-        device = self._device_from_ip_address(ip_address, command)
+    def osc_ue4_launch_confirm(self, address, command):
+        device = self._device_from_address(address, command)
         if not device:
             return
 
@@ -2093,26 +2097,26 @@ class SwitchboardDialog(QtCore.QObject):
         # Set the device status to ready
         device.status = DeviceStatus.READY
 
-    def osc_add_send_target_confirm(self, ip_address, command, value):
-        device = self.device_manager.device_with_ip_address(ip_address[0])
+    def osc_add_send_target_confirm(self, address, command, value):
+        device = self.device_manager.device_with_address(address[0])
         if not device:
             return
 
         device.osc_add_send_target_confirm()
 
-    def osc_arsession_stop_confrim(self, ip_address, command, value):
-        LOGGER.debug(f'osc_arsession_stop_confrim {value}')
+    def osc_arsession_stop_confirm(self, address, command, value):
+        LOGGER.debug(f'osc_arsession_stop_confirm {value}')
 
-    def osc_arsession_start_confrim(self, ip_address, command, value):
-        device = self._device_from_ip_address(ip_address, command, value=value)
+    def osc_arsession_start_confirm(self, address, command, value):
+        device = self._device_from_address(address, command, value=value)
         if not device:
             return
 
         device.connect_listener()
 
-    def osc_battery(self, ip_address, command, value):
+    def osc_battery(self, address, command, value):
         # The Battery command is used to handshake with LiveLinkFace. Don't reject it if it's not connected 
-        device = self.device_manager.device_with_ip_address(ip_address[0])
+        device = self.device_manager.device_with_address(address[0])
         
         if not device:
             return
@@ -2124,8 +2128,8 @@ class SwitchboardDialog(QtCore.QObject):
         device_widget = self.device_list_widget.device_widget_by_hash(device.device_hash)
         device_widget.set_battery(value)
 
-    def osc_data(self, ip_address, command, value):
-        device = self._device_from_ip_address(ip_address, command)
+    def osc_data(self, address, command, value):
+        device = self._device_from_address(address, command)
         if not device:
             return
 
diff --git a/Engine/Plugins/VirtualProduction/Switchboard/Source/Switchboard/switchboard/switchboard_widgets.py b/Engine/Plugins/VirtualProduction/Switchboard/Source/Switchboard/switchboard/switchboard_widgets.py
--- a/Engine/Plugins/VirtualProduction/Switchboard/Source/Switchboard/switchboard/switchboard_widgets.py
+++ b/Engine/Plugins/VirtualProduction/Switchboard/Source/Switchboard/switchboard/switchboard_widgets.py
@@ -1,4 +1,6 @@
 # Copyright Epic Games, Inc. All Rights Reserved.
+import re
+import socket
 import time
 from typing import Optional
 
@@ -9,16 +11,16 @@ from switchboard.switchboard_logging import LOGGER
 
 DEVICE_LIST_WIDGET_HEIGHT = 54
 DEVICE_HEADER_LIST_WIDGET_HEIGHT = 40
-DEVICE_WIDGET_HIDE_IP_ADDRESS_WIDTH = 500
+DEVICE_WIDGET_HIDE_ADDRESS_WIDTH = 500
 
 
 class NonScrollableComboBox(QtWidgets.QComboBox):
     onHoverScrollBox = QtCore.Signal()
-    
+
     def __init__(self, parent=None):
         super().__init__(parent=parent)
         self.installEventFilter(self)
-        
+
     def enterEvent(self, event):
         self.onHoverScrollBox.emit()
         super().enterEvent(event)
@@ -128,10 +130,10 @@ class MultiSelectionComboBox(QtWidgets.QComboBox):
 # A combo box that has a static icon and opens a list of options when pressed, e.g. view options in settings
 class DropDownMenuComboBox(QtWidgets.QComboBox):
     on_select_option = QtCore.Signal(str)
-    
+
     def __init__(self, icon: QtGui.QIcon = None, icon_size: int = 25, parent=None):
         super().__init__(parent=parent)
-        
+
         drop_down_arrow_size = 15
         if icon is not None:
             self.addItem(icon, "")
@@ -139,21 +141,21 @@ class DropDownMenuComboBox(QtWidgets.QComboBox):
         else:
             self.addItem("")
             self.setFixedWidth(drop_down_arrow_size)
-        
+
         self.model().setData(self.model().index(0, 0), QtCore.QSize(100, 100), QtCore.Qt.SizeHintRole)
         self.view().setRowHidden(0, True)
         self.currentIndexChanged.connect(self._on_index_changed)
-    
+
     def _on_index_changed(self, index: int):
         selected_item = self.itemText(index)
         self.on_select_option.emit(selected_item)
         # Always make sure that the icon is shown
         self.setCurrentIndex(0)
-    
+
     def showPopup(self):
         self.view().setMinimumWidth(self.view().sizeHintForColumn(0))
         super().showPopup()
-    
+
 
 class ControlQPushButton(QtWidgets.QPushButton):
     def __init__(self, parent = None, *, hover_focus: bool = True):
@@ -299,13 +301,12 @@ class FramelessQLineEdit(QtWidgets.QLineEdit):
         elif e.key() == QtCore.Qt.Key_Escape:
             self.setText(self.current_text)
             self.clearFocus()
-    
-    
+
 
 class CollapsibleGroupBox(QtWidgets.QGroupBox):
     def __init__(self):
         super().__init__()
-        
+
         self.setCheckable(True)
         self.setStyleSheet(
             'QGroupBox::indicator:checked:hover {image: url(:icons/images/tree_arrow_expanded_hovered.png);}'
@@ -316,18 +317,18 @@ class CollapsibleGroupBox(QtWidgets.QGroupBox):
 
         self.toggled.connect(self._on_set_expanded)
         self.setChecked(True)
-        
+
     @property
     def is_expanded(self):
         return self.isChecked()
-        
+
     def set_expanded(self, value: bool):
         if value and not self.is_expanded or not value and self.is_expanded:
             self._on_set_expanded(value)
             self.blockSignals(True)
             self.setChecked(value)
             self.blockSignals(False)
-        
+
     def _on_set_expanded(self, should_expand: bool):
         if should_expand:
             self.setMaximumHeight(self._original_maximum_height)
@@ -336,20 +337,20 @@ class CollapsibleGroupBox(QtWidgets.QGroupBox):
             # Just font height does not suffice ... need to investigate how to get a better value programmatically
             safety_margin = 6
             self.setMaximumHeight(self.fontMetrics().height() + safety_margin)
-            
+
 
 class SearchableComboBox(QtWidgets.QComboBox):
-    
+
     class CustomQList(QtWidgets.QListWidget):
         def __init__(self, parent):
             super().__init__(parent)
-            
-            
+
+
     def __init__(self, parent):
         super().__init__(parent)
-        
+
         self.HEIGHT = 15
-        
+
         self.item_list = self.CustomQList(self)
         self.item_list.setViewportMargins(0, self.HEIGHT, 0, 0)
         self.setView(self.item_list)
@@ -358,36 +359,157 @@ class SearchableComboBox(QtWidgets.QComboBox):
         self.search_line = QtWidgets.QLineEdit(self.view())
         self.search_line.setVisible(False)
         self.search_line.setPlaceholderText("Search")
-       
+
         self.search_line.textChanged.connect(self.text_changed)
-        
+
     def showPopup(self):
-            
         super().showPopup()
-        
+
         MARGIN = 10
         XPOS = MARGIN/2
         YPOS = 2
-        
-        self.search_line.setGeometry(XPOS, YPOS, self.view().width() - MARGIN*2.5, self.HEIGHT)
+
+        self.search_line.setGeometry(
+            XPOS, YPOS, self.view().width() - MARGIN*2.5, self.HEIGHT)
         self.search_line.setVisible(True)
-        
+
         self.CB_item_text = [self.itemText(i) for i in range(self.count())]
-        
+
     def hidePopup(self):
         super().hidePopup()
-        
+
         self.search_line.setVisible(False)
         self.search_line.clear()
-    
-    def text_changed(self, text:str):
-        
+
+    def text_changed(self, text: str):
         for item in self.CB_item_text:
             if item.lower().__contains__(text.lower()):
                 self.view().setRowHidden(self.findText(item), False)
             else:
                 self.view().setRowHidden(self.findText(item), True)
-     
+
+
+class HostnameValidator(QtGui.QValidator):
+    INVALID_CHARS_RE = re.compile(r"[^A-Z.\d-]", re.IGNORECASE)
+    VALID_SEGMENT_RE = re.compile(r"(?!-)[A-Z\d-]{1,63}(?<!-)$",
+                                  re.IGNORECASE)
+
+    def fixup(self, input: str):
+        return self.INVALID_CHARS_RE.sub('', input)
+
+    def validate(self, input: str, pos: int):
+        input_len = len(input)
+        if input_len == 0 or input_len > 255:
+            return QtGui.QValidator.Intermediate
+
+        if input[-1] == '.':
+            input = input[:-1]  # trim FQDN terminator
+
+        for segment in input.split('.'):
+            if not self.VALID_SEGMENT_RE.match(segment):
+                return QtGui.QValidator.Intermediate
+
+        return QtGui.QValidator.Acceptable
+
+
+class AddressComboBox(NonScrollableComboBox):
+    '''
+    An editable combo box with visible indication of hostname resolvability.
+    '''
+
+    class ResolverRunnable(QtCore.QRunnable):
+        ''' Performs a name resolution query on the thread pool. '''
+
+        class Signals(QtCore.QObject):
+            result = QtCore.Signal(list)
+
+        def __init__(self, name: str):
+            super().__init__()
+            self.name = name
+            self.signals = self.Signals()
+
+        def run(self):
+            addrs = list[str]()
+            try:
+                (_, _, addrs) = socket.gethostbyname_ex(self.name)
+            except socket.gaierror:
+                pass
+            self.signals.result.emit(addrs)
+
+    def __init__(self, parent=None):
+        super().__init__(parent=parent)
+
+        self.runnable = AddressComboBox.ResolverRunnable('')
+        self.resolving = False
+        self.debounce = QtCore.QTimer()
+        self.debounce.setSingleShot(True)
+        self.debounce.timeout.connect(self.on_text_settled)
+
+        self.setEditable(True)
+        self.lineEdit().setValidator(HostnameValidator())
+        self.lineEdit().textChanged.connect(self.on_text_changed)
+        self.lineEdit().editingFinished.connect(self.on_editing_finished)
+
+    def start_resolve(self):
+        if self.resolving:
+            return False
+
+        try:
+            self.runnable.signals.result.disconnect(self.on_resolver_result)
+        except RuntimeError:
+            pass  # non-existent signal connections raise (seems excessive)
+
+        name = self.lineEdit().text()
+        self.runnable = AddressComboBox.ResolverRunnable(name)
+        self.runnable.signals.result.connect(self.on_resolver_result)
+        QtCore.QThreadPool.globalInstance().start(self.runnable)
+        LOGGER.debug(f'starting resolve: "{name}"')
+        self.resolving = True
+        return True
+
+    def on_text_changed(self, new_text: str):
+        if not self.lineEdit().hasAcceptableInput():
+            set_qt_property(self, 'validation', 'invalid')
+            self.debounce.stop()
+            return
+
+        set_qt_property(self, 'validation', 'intermediate')
+
+        # Don't issue the query until the user stops typing for a second
+        SETTLE_WAIT_MILLISEC = 1000
+        self.debounce.start(SETTLE_WAIT_MILLISEC)
+
+    def on_editing_finished(self):
+        self.debounce.stop()
+
+        if not self.lineEdit().hasAcceptableInput():
+            set_qt_property(self, 'validation', 'invalid')
+            return
+
+        set_qt_property(self, 'validation', 'intermediate')
+        self.start_resolve()
+
+    def on_text_settled(self):
+        self.start_resolve()
+
+    def on_resolver_result(self, addrs: list[str]):
+        self.resolving = False
+
+        current_text = self.lineEdit().text()
+        last_resolved = self.runnable.name
+
+        if current_text != last_resolved:
+            # Resolver result was stale, issue a new one
+            self.start_resolve()
+            return
+
+        LOGGER.debug(f'resolved: "{", ".join(addrs)}"')
+
+        if len(addrs) == 0:
+            set_qt_property(self, 'validation', 'invalid')
+        else:
+            set_qt_property(self, 'validation', 'succeeded')
+
 
 def set_qt_property(
     widget: QtWidgets.QWidget, prop, value, *,
