Switchboard: add a setting to allow finding levels that are inside project plugins

This adds a "maps_plugin_filters" StringListSetting that can be used to add
levels located inside plugins to the "Levels" dropdown menu in the main
Switchboard dialog. The setting is a list of filters similar to the existing "maps_filter"
setting which tests against map paths, but "maps_plugin_filters" is tested against
plugin names. Any plugins whose name matches any of the filters will be searched
for levels when populating the dropdown menu.

diff --git a/Engine/Plugins/VirtualProduction/Switchboard/Source/Switchboard/switchboard/config.py b/Engine/Plugins/VirtualProduction/Switchboard/Source/Switchboard/switchboard/config.py
--- a/Engine/Plugins/VirtualProduction/Switchboard/Source/Switchboard/switchboard/config.py
+++ b/Engine/Plugins/VirtualProduction/Switchboard/Source/Switchboard/switchboard/config.py
@@ -2062,7 +2062,14 @@ class Config(object):
                 data.get('maps_filter', '*.umap'),
                 placeholder_text="*.umap",
                 tool_tip="Walk every file in the Map Path and run a fnmatch to filter the file names"
-            )
+            ),
+            'maps_plugin_filters': StringListSetting(
+                "maps_plugin_filters",
+                "Map Plugin Filters",
+                data.get('maps_plugin_filters', []),
+                tool_tip="Plugins whose name matches any of these filters will also be searched for maps.",
+                migrate_data=migrate_comma_separated_string_to_list
+            ),
         }
 
         self.PROJECT_NAME = self.basic_project_settings["project_name"]
@@ -2071,6 +2078,7 @@ class Config(object):
         self.BUILD_ENGINE = self.basic_project_settings["build_engine"]
         self.MAPS_PATH = self.basic_project_settings["maps_path"]
         self.MAPS_FILTER = self.basic_project_settings["maps_filter"]
+        self.MAPS_PLUGIN_FILTERS = self.basic_project_settings["maps_plugin_filters"]
 
         self.osc_settings = {
             "osc_server_port": IntSetting(
@@ -2352,6 +2360,7 @@ class Config(object):
         data['build_engine'] = self.BUILD_ENGINE.get_value()
         data["maps_path"] = self.MAPS_PATH.get_value()
         data["maps_filter"] = self.MAPS_FILTER.get_value()
+        data["maps_plugin_filters"] = self.MAPS_PLUGIN_FILTERS.get_value()
         data["listener_exe"] = self.LISTENER_EXE.get_value()
 
         self.save_unreal_insights(data)
@@ -2445,35 +2454,127 @@ class Config(object):
         path_name = path.replace(self.get_project_dir(), '', 1)
         path_name = path_name.replace(os.sep, '/')
         return path_name
-        
-    def get_project_dir(self):
+
+    def get_project_dir(self) -> str:
+        '''
+        Get the root directory of the project.
+
+        This is the directory in which the .uproject file lives.
+        '''
         return os.path.dirname(self.UPROJECT_PATH.get_value().replace('"', ''))
 
+    def get_project_content_dir(self) -> str:
+        '''
+        Get the "Content" directory of the project.
+        '''
+        return os.path.join(self.get_project_dir(), 'Content')
+
+    def get_project_plugins_dir(self) -> str:
+        '''
+        Get the "Plugins" directory of the project where all of the
+        project-based plugins are located.
+        '''
+        return os.path.join(self.get_project_dir(), 'Plugins')
+
+    def get_project_plugin_names(self) -> typing.List[str]:
+        '''
+        Get a list of the names of all project-based plugins in the project.
+        '''
+        project_plugins_path = os.path.normpath(self.get_project_plugins_dir())
+        if not os.path.isdir(project_plugins_path):
+            return []
+
+        plugin_names = [
+            x.name for x in os.scandir(project_plugins_path) if x.is_dir()]
+
+        return plugin_names
+
+    def get_project_plugin_content_dir(self, plugin_name: str) -> str:
+        '''
+        Get the "Content" directory of the project-based plugin with the
+        given name.
+        '''
+        return os.path.join(
+            self.get_project_plugins_dir(), plugin_name, 'Content')
+
+    def plugin_name_matches(self, plugin_name: str, name_filters: typing.List[str]) -> bool:
+        '''
+        Test whether a plugin name matches any of the provided filters.
+
+        Returns True if any one of the filters matches, or False otherwise.
+        '''
+        for name_filter in name_filters:
+            if fnmatch.fnmatch(plugin_name, name_filter):
+                return True
+
+        return False
+
+    def resolve_content_path(self, file_path: str, plugin_name: str = None) -> str:
+        '''
+        Resolve a file path on the file system to the corresponding content
+        path in UE.
+
+        If a plugin_name is provided, the file is assumed to live inside that
+        plugin and its content path will have the appropriate plugin
+        name-based prefix. Otherwise, the file is assumed to live inside the
+        project's content folder.
+        '''
+        if plugin_name:
+            content_dir = self.get_project_plugin_content_dir(plugin_name)
+            ue_path_prefix = f'/{plugin_name}'
+        else:
+            content_dir = self.get_project_content_dir()
+            ue_path_prefix = '/Game'
+
+        path_name = file_path.replace(content_dir, ue_path_prefix, 1)
+        path_name = self.shrink_path(path_name)
+
+        return path_name
+
     def maps(self):
         '''
-        Returns a list of ful map paths in an Unreal Engine project such as [ "/Game/Maps/MapName" ].
-        It will always start with Game and the slashes will always be "/" independent of the platform's separator.
+        Returns a list of full map paths in an Unreal Engine project and
+        in project-based plugins such as:
+            [
+                "/Game/Maps/MapName",
+                "/MyPlugin/Levels/MapName"
+            ]
+        The slashes will always be "/" independent of the platform's separator.
         '''
-        maps_path = os.path.normpath(
+        content_maps_path = os.path.normpath(
             os.path.join(
-                self.get_project_dir(),
-                'Content',
+                self.get_project_content_dir(),
                 self.MAPS_PATH.get_value()))
 
+        # maps_paths stores a list of tuples of the form
+        # (plugin_name, file_path). This allows us to differentiate between
+        # maps in the project and maps in a plugin.
+        maps_paths = [(None, content_maps_path)]
+
+        maps_plugin_filters = self.MAPS_PLUGIN_FILTERS.get_value()
+        if maps_plugin_filters:
+            plugin_names = self.get_project_plugin_names()
+
+            for plugin_name in plugin_names:
+                if self.plugin_name_matches(plugin_name, maps_plugin_filters):
+                    maps_paths.append(
+                        (plugin_name, self.get_project_plugin_content_dir(plugin_name)))
+
         maps = []
-        for path_to_map, b, files in os.walk(maps_path):
-            for name in files:
-                if not fnmatch.fnmatch(name, self.MAPS_FILTER.get_value()):
-                    continue
-                
-                map_name, _ = os.path.splitext(name)
-                # Ignore the fact that maps may exist in plugins - we only search game content
-                path_name = path_to_map.replace('Content', 'Game', 1)
-                path_name = os.path.join(path_name, map_name)
-                path_name = self.shrink_path(path_name)
-                
-                if path_name not in maps:
-                    maps.append(path_name)
+        for (plugin_name, maps_path) in maps_paths:
+            for dirpath, b, file_names in os.walk(maps_path):
+                for file_name in file_names:
+                    if not fnmatch.fnmatch(file_name, self.MAPS_FILTER.get_value()):
+                        continue
+
+                    map_name, _ = os.path.splitext(file_name)
+                    file_path_to_map = os.path.join(dirpath, map_name)
+
+                    content_path_to_map = self.resolve_content_path(
+                        file_path_to_map, plugin_name=plugin_name)
+
+                    if content_path_to_map not in maps:
+                        maps.append(content_path_to_map)
 
         maps.sort()
         return maps
