Switchboard: Added LogDisplayClusterMedia to default log options.

diff --git a/Engine/Plugins/VirtualProduction/Switchboard/Source/Switchboard/switchboard/devices/ndisplay/plugin_ndisplay.py b/Engine/Plugins/VirtualProduction/Switchboard/Source/Switchboard/switchboard/devices/ndisplay/plugin_ndisplay.py
--- a/Engine/Plugins/VirtualProduction/Switchboard/Source/Switchboard/switchboard/devices/ndisplay/plugin_ndisplay.py
+++ b/Engine/Plugins/VirtualProduction/Switchboard/Source/Switchboard/switchboard/devices/ndisplay/plugin_ndisplay.py
@@ -493,6 +493,7 @@ class DevicenDisplay(DeviceUnreal):
                 'LogDisplayClusterRender',
                 'LogDisplayClusterRenderSync',
                 'LogDisplayClusterViewport',
+                'LogDisplayClusterMedia',
                 'LogLiveLink',
                 'LogRemoteControl',
             ],
