Switchboard: Fix resetting uproject or engine dir path when cancelling browse in the create config dialog.

#skipci

diff --git a/Engine/Plugins/VirtualProduction/Switchboard/Source/Switchboard/switchboard/add_config_dialog.py b/Engine/Plugins/VirtualProduction/Switchboard/Source/Switchboard/switchboard/add_config_dialog.py
--- a/Engine/Plugins/VirtualProduction/Switchboard/Source/Switchboard/switchboard/add_config_dialog.py
+++ b/Engine/Plugins/VirtualProduction/Switchboard/Source/Switchboard/switchboard/add_config_dialog.py
@@ -531,7 +531,13 @@ class AddConfigDialog(QtWidgets.QDialog):
         self.config_path_line_edit.setText(config_path_str)
 
     def on_browse_uproject_path(self, uproject_search_path):
-        self.uproject, _ = QtWidgets.QFileDialog.getOpenFileName(self, "Select uProject file", self.engine_dir, "uProject (*.uproject)")
+
+        new_file, _ = QtWidgets.QFileDialog.getOpenFileName(self, "Select uProject file", self.engine_dir, "uProject (*.uproject)")
+
+        if not new_file:
+            return
+
+        self.uproject = new_file
         self.uproject = os.path.normpath(self.uproject)
         self.uproject_line_edit.setText(self.uproject)
 
@@ -540,7 +546,12 @@ class AddConfigDialog(QtWidgets.QDialog):
         self.update_button_box()
 
     def on_browse_engine_dir(self):
-        self.engine_dir = QtWidgets.QFileDialog.getExistingDirectory(self, "Select UE 'Engine' directory")
+        new_dir = QtWidgets.QFileDialog.getExistingDirectory(self, "Select UE 'Engine' directory")
+
+        if not new_dir:
+            return
+
+        self.engine_dir = new_dir
         self.engine_dir = os.path.normpath(self.engine_dir)
         self.engine_dir_line_edit.setText(self.engine_dir)
 
