Detect the running arguments for the server and reflect information into the UI

#skipci

diff --git a/Engine/Plugins/VirtualProduction/Switchboard/Source/Switchboard/switchboard/switchboard_application.py b/Engine/Plugins/VirtualProduction/Switchboard/Source/Switchboard/switchboard/switchboard_application.py
--- a/Engine/Plugins/VirtualProduction/Switchboard/Source/Switchboard/switchboard/switchboard_application.py
+++ b/Engine/Plugins/VirtualProduction/Switchboard/Source/Switchboard/switchboard/switchboard_application.py
@@ -2,10 +2,12 @@
 
 import os
 import pathlib
+import re
 import subprocess
 import sys
 import tempfile
 import threading
+
 from typing import Dict, List, NamedTuple, Optional
 
 import pythonosc.dispatcher
@@ -15,7 +17,6 @@ from .config import CONFIG, SETTINGS
 from .switchboard_logging import LOGGER
 from . import switchboard_utils as sb_utils
 
-
 class OscServer:
     def __init__(self):
         self.server_thread = None
@@ -90,6 +91,7 @@ class MultiUserApplication:
 
         # Application Options
         self.concert_ignore_cl = False
+        self.clear_process_info()
 
     def exe_path(self):
         return CONFIG.multiuser_server_path()
@@ -104,10 +106,8 @@ class MultiUserApplication:
         return ''
 
     def get_mu_server_endpoint_arg(self):
-        endpoint_setting = CONFIG.MUSERVER_ENDPOINT.get_value().strip()
-        if endpoint_setting:
-            addr_setting = SETTINGS.ADDRESS.get_value().strip()
-            endpoint = sb_utils.expand_endpoint(endpoint_setting, addr_setting)
+        endpoint = self.endpoint_address()
+        if endpoint:
             return f'-UDPMESSAGING_TRANSPORT_UNICAST="{endpoint}"'
         return ''
 
@@ -171,12 +171,67 @@ class MultiUserApplication:
         else:
             self.poll_process().kill()
 
-    def is_running(self):
-        if self.process and (self.process.poll() is None):
+    def extract_process_info(self):
+        pid = self.poll_process().get_pid()
+        command_line = self.poll_process().get_command_line()
+        if command_line == '' or pid == self._running_pid:
+            return
+
+        find_ip_re =  re.compile( r'-UDPMESSAGING_TRANSPORT_UNICAST="(\d+.\d+.\d+.\d+:\d+)"')
+        find_name_re = re.compile(r'-CONCERTSERVER="([A-Za-z0-9_]+)"')
+
+        self.clear_process_info()
+        self._running_pid = pid
+        try:
+            ip_match = find_ip_re.search(command_line)
+            name_match = find_name_re.search(command_line)
+            if ip_match:
+                self._running_endpoint = ip_match.group(1)
+            if name_match:
+                self._running_name = name_match.group(1)
+        except:
+            pass
+
+
+    def validate_process(self):
+        if self._running_pid is None:
+            return False
+
+        endpoint = self.endpoint_address()
+        server_name = CONFIG.MUSERVER_SERVER_NAME.get_value()
+        if endpoint == self._running_endpoint and server_name == self._running_name:
             return True
-        elif self.poll_process().poll() is None:
+
+        return False
+
+    def running_server_name(self):
+        return self._running_name
+
+    def running_endpoint(self):
+        return self._running_endpoint
+
+    def endpoint_address(self):
+        endpoint_setting = CONFIG.MUSERVER_ENDPOINT.get_value().strip()
+        endpoint = ""
+        if endpoint_setting:
+            addr_setting = SETTINGS.ADDRESS.get_value().strip()
+            endpoint = sb_utils.expand_endpoint(endpoint_setting, addr_setting)
+        return endpoint
+
+    def server_name(self):
+        return CONFIG.MUSERVER_SERVER_NAME.get_value()
+
+    def clear_process_info(self):
+        self._running_endpoint = ""
+        self._running_name = ""
+        self._running_pid = None
+
+    def is_running(self):
+        if self.poll_process().poll() is None:
+            self.extract_process_info()
             return True
 
+        self.clear_process_info()
         return False
 
 
diff --git a/Engine/Plugins/VirtualProduction/Switchboard/Source/Switchboard/switchboard/switchboard_dialog.py b/Engine/Plugins/VirtualProduction/Switchboard/Source/Switchboard/switchboard/switchboard_dialog.py
--- a/Engine/Plugins/VirtualProduction/Switchboard/Source/Switchboard/switchboard/switchboard_dialog.py
+++ b/Engine/Plugins/VirtualProduction/Switchboard/Source/Switchboard/switchboard/switchboard_dialog.py
@@ -169,6 +169,25 @@ class DeviceAdditionalSettingsUI(QtCore.QObject):
         self.assign_button(button,parent)
         return button
 
+class PeriodicRunnable(QtCore.QRunnable):
+    ''' Performs periodic tasks on the switchboard dialog. '''
+
+    def __init__(self, switchboard):
+        super().__init__()
+        self._switchboard = switchboard
+        self._exiting = False
+
+    def exit(self):
+        self._exiting = True
+
+    def run(self):
+        while not self._exiting:
+            self._switchboard.update_muserver_button()
+            self._switchboard.update_locallistener_menuitem()
+            self._switchboard.update_insights_menuitem()
+
+            time.sleep(1.0)
+
 class SwitchboardDialog(QtCore.QObject):
 
     STYLESHEET_PATH = os.path.join(RELATIVE_PATH, 'ui/switchboard.qss')
@@ -227,7 +246,7 @@ class SwitchboardDialog(QtCore.QObject):
 
         self.init_stylesheet_watcher()
 
-        self._exiting = False
+        self._periodic_runner = None
         self._shoot = None
         self._sequence = None
         self._slate = None
@@ -434,6 +453,7 @@ class SwitchboardDialog(QtCore.QObject):
         self.refresh_window_title()
 
         self.script_manager.on_postinit(self)
+        self.have_warned_about_muserver = False
         
     def _try_change_address(self):
         new_value = self.window.current_address_value.text()
@@ -454,15 +474,30 @@ class SwitchboardDialog(QtCore.QObject):
             SETTINGS.ADDRESS.update_value(old_address)
             SETTINGS.save()
             self.osc_server.launch(SETTINGS.ADDRESS.get_value(), CONFIG.OSC_SERVER_PORT.get_value())
-            
 
-    def _poll_muserver_status(self):
-        '''
-        Poll the status of the Multi-user server on a 1 second timer.
-        '''
-        while not self._exiting:
-            self.update_muserver_button()
-            time.sleep(1.0)
+    def warn_user_about_muserver(self):
+        if self.have_warned_about_muserver:
+            return
+
+        ServerInstance = switchboard_application.get_multi_user_server_instance()
+        expected_server = ServerInstance.server_name()
+        expected_endpoint = ServerInstance.endpoint_address()
+        actual_server = ServerInstance.running_server_name()
+        actual_endpoint = ServerInstance.running_endpoint()
+        LOGGER.warning(f'The running Multi-user server does not match Switchboard configuration. Expected "{expected_server}" with "{expected_endpoint}" but found "{actual_server}" with "{actual_endpoint}". Please restart the multi-user server.')
+
+        self.have_warned_about_muserver = True
+
+    def get_muserver_label(self, is_running, is_valid):
+        ServerInstance = switchboard_application.get_multi_user_server_instance()
+        if is_running and is_valid:
+            actual_server = ServerInstance.running_server_name()
+            actual_endpoint = ServerInstance.running_endpoint()
+            return f'Multi-user Server ({actual_server} {actual_endpoint})'
+        elif is_running and not is_valid:
+            return 'Multi-user Server (WARNING: RUNNING SERVER DOES NOT MATCH CONFIGURATION)'
+        else:
+            return 'Multi-user Server'
 
     def update_muserver_button(self):
         '''
@@ -471,10 +506,17 @@ class SwitchboardDialog(QtCore.QObject):
         ServerInstance = switchboard_application.get_multi_user_server_instance()
         is_checked = self.window.muserver_start_stop_button.isChecked()
         is_running = ServerInstance.is_running()
+        is_valid =  ServerInstance.validate_process()
         if  (is_running or self._started_mu_server) and not is_checked:
             self.window.muserver_start_stop_button.setChecked( True )
         elif not is_running and is_checked:
             self.window.muserver_start_stop_button.setChecked( False )
+            self.have_warned_about_muserver = False
+
+        label = self.get_muserver_label(is_running or self._started_mu_server,is_valid)
+        self.window.muserver_label.setText(label)
+        if is_running and not is_valid:
+            self.warn_user_about_muserver()
 
         if is_running and self._started_mu_server:
             # Server has finished starting so reset our start flag.
@@ -484,19 +526,8 @@ class SwitchboardDialog(QtCore.QObject):
         '''
         Sets up a thread for performing maintenance tasks
         '''
-        thread = threading.Thread(target=self._do_periodic_tasks, args=[], kwargs={})
-        thread.start()
-
-    def _do_periodic_tasks(self):
-        ''' Performs periodic tasks, like updating certain parts of the UI '''
-
-        while not self._exiting:
-
-            self.update_muserver_button()
-            self.update_locallistener_menuitem()
-            self.update_insights_menuitem()
-
-            time.sleep(1.0)
+        self._periodic_runner = PeriodicRunnable(self)
+        QtCore.QThreadPool.globalInstance().start(self._periodic_runner)
 
     def update_locallistener_menuitem(self):
         ''' 
@@ -695,7 +726,8 @@ class SwitchboardDialog(QtCore.QObject):
                    for device in self.device_manager.devices())
 
     def on_exit(self):
-        self._exiting = True
+        if self._periodic_runner:
+            self._periodic_runner.exit()
         self.osc_server.close()
         for device in self.device_manager.devices():
             device.disconnect_listener()
diff --git a/Engine/Plugins/VirtualProduction/Switchboard/Source/Switchboard/switchboard/switchboard_utils.py b/Engine/Plugins/VirtualProduction/Switchboard/Source/Switchboard/switchboard/switchboard_utils.py
--- a/Engine/Plugins/VirtualProduction/Switchboard/Source/Switchboard/switchboard/switchboard_utils.py
+++ b/Engine/Plugins/VirtualProduction/Switchboard/Source/Switchboard/switchboard/switchboard_utils.py
@@ -6,6 +6,8 @@ import pathlib
 import subprocess
 import sys
 import threading
+import io
+import csv
 
 from .switchboard_logging import LOGGER
 
@@ -110,6 +112,24 @@ class PollProcess(object):
     def __init__(self, task_name: str):
         self.task_name = task_name
 
+    def get_command_line(self):
+        get_commandline_cmd = f"wmic process where caption=\"{self.task_name}\" get commandline"
+        try:
+            output = subprocess.check_output(get_commandline_cmd, startupinfo=get_hidden_sp_startupinfo()).decode()
+            return output
+        except:
+            return ""
+
+    def get_pid(self):
+        pid_cmd = f"tasklist /FI \"IMAGENAME eq {self.task_name}\" /FO csv"
+        try:
+            output = subprocess.check_output(pid_cmd, startupinfo=get_hidden_sp_startupinfo()).decode()
+            dictobj = next(csv.DictReader(io.StringIO(output)))
+            if dictobj:
+                return dictobj['PID']
+        except:
+            return "Unknown"
+
     def poll(self):
         # 'list' output format because default 'table' truncates imagename to 25 characters
         tasklist_cmd = f"tasklist /FI \"IMAGENAME eq {self.task_name}\" /FO list"
