Switchboard: Inclusive terminology fixes.

diff --git a/Engine/Plugins/VirtualProduction/Switchboard/Source/Switchboard/switchboard/config_osc.py b/Engine/Plugins/VirtualProduction/Switchboard/Source/Switchboard/switchboard/config_osc.py
--- a/Engine/Plugins/VirtualProduction/Switchboard/Source/Switchboard/switchboard/config_osc.py
+++ b/Engine/Plugins/VirtualProduction/Switchboard/Source/Switchboard/switchboard/config_osc.py
@@ -1,4 +1,5 @@
 # Copyright Epic Games, Inc. All Rights Reserved.
+
 UE4_LAUNCH = '/UE4Launch'
 UE4_LAUNCH_QUERY = '/UE4LaunchQuery'
 UE4_LAUNCH_CONFIRM = '/UE4LaunchConfirm'
@@ -21,7 +22,7 @@ CLEAR_LIVE_LINK_ADDRESSES = '/ClearAllLiveLinkAddresses'
 LIVE_LINK_STREAM_START = '/LiveLinkStreamStart'
 LIVE_LINK_STREAM_STOP = '/LiveLinkStreamStop'
 
-# Master Recording commands
+# Recording commands
 RECORD_START = '/RecordStart'
 RECORD_START_CONFIRM = '/RecordStartConfirm'
 RECORD_STOP = '/RecordStop'
@@ -50,4 +51,4 @@ VPRC_RECIEVER_NAME = 'Record'
 
 # MONITOR
 MONITOR_WILDCARD = '/Monitor*'
-MONITOR_RECIEVER_NAME = 'Monitor'
\ No newline at end of file
+MONITOR_RECIEVER_NAME = 'Monitor'
diff --git a/Engine/Plugins/VirtualProduction/Switchboard/Source/Switchboard/switchboard/devices/ndisplay/ndisplay_monitor.py b/Engine/Plugins/VirtualProduction/Switchboard/Source/Switchboard/switchboard/devices/ndisplay/ndisplay_monitor.py
--- a/Engine/Plugins/VirtualProduction/Switchboard/Source/Switchboard/switchboard/devices/ndisplay/ndisplay_monitor.py
+++ b/Engine/Plugins/VirtualProduction/Switchboard/Source/Switchboard/switchboard/devices/ndisplay/ndisplay_monitor.py
@@ -107,7 +107,7 @@ class nDisplayMonitor(QAbstractTableModel):
             return self.COLOR_NORMAL if is_good else self.COLOR_WARNING
 
         if colname == 'Displays':
-            is_normal = (('Slave' in value or 'Master' in value)
+            is_normal = (('Follower' in value or 'Leader' in value)
                          and ('Unsynced' not in value))
             return self.COLOR_NORMAL if is_normal else self.COLOR_WARNING
 
@@ -132,6 +132,7 @@ class nDisplayMonitor(QAbstractTableModel):
                     # the mechanism for retrieving it seems fragile.
                     # Based on: https://docs.microsoft.com/en-us/windows/release-health/release-information  # noqa
                     BUILD_TO_SDK_VERSION = {
+                        '19044': '21H2',
                         '19043': '21H1',
                         '19042': '20H2',
                         '19041': '2004',
@@ -360,19 +361,20 @@ class nDisplayMonitor(QAbstractTableModel):
                         for sync_topo in sync_topos]
         sync_sources = [source_str.get(sync_source, 'Unknown')
                         for sync_source in sync_sources]
-        bInternalSlaves = [sync_topo['syncStatusParams']['bInternalSlave']
-                           for sync_topo in sync_topos]
+        bInternalSecondaries = [
+            sync_topo['syncStatusParams']['bInternalSecondary']
+            for sync_topo in sync_topos]
 
-        sync_slaves = []
+        sync_followers = []
 
         for i in range(len(sync_sources)):
-            if bInternalSlaves[i] and sync_sources[i] == 'Vsync':
-                sync_slaves.append('Vsync(daisy)')
+            if bInternalSecondaries[i] and sync_sources[i] == 'Vsync':
+                sync_followers.append('Vsync(daisy)')
             else:
-                sync_slaves.append(sync_sources[i])
+                sync_followers.append(sync_sources[i])
 
-        if len(sync_slaves) > 0:
-            data['SyncSource'] = '\n'.join(sync_slaves)
+        if len(sync_followers) > 0:
+            data['SyncSource'] = '\n'.join(sync_followers)
         else:
             data['SyncSource'] = self.DATA_MISSING
 
diff --git a/Engine/Plugins/VirtualProduction/Switchboard/Source/Switchboard/switchboard/devices/unreal/version_helpers.py b/Engine/Plugins/VirtualProduction/Switchboard/Source/Switchboard/switchboard/devices/unreal/version_helpers.py
--- a/Engine/Plugins/VirtualProduction/Switchboard/Source/Switchboard/switchboard/devices/unreal/version_helpers.py
+++ b/Engine/Plugins/VirtualProduction/Switchboard/Source/Switchboard/switchboard/devices/unreal/version_helpers.py
@@ -3,7 +3,7 @@
 import typing
 
 
-LISTENER_COMPATIBLE_VERSION = (2,0)
+LISTENER_COMPATIBLE_VERSION = (2,1)
 
 # Version where redeploy support was initially introduced.
 LISTENER_MIN_REDEPLOY_VERSION = (1,4,0)
diff --git a/Engine/Source/Programs/SwitchboardListener/Private/SwitchboardListener.cpp b/Engine/Source/Programs/SwitchboardListener/Private/SwitchboardListener.cpp
--- a/Engine/Source/Programs/SwitchboardListener/Private/SwitchboardListener.cpp
+++ b/Engine/Source/Programs/SwitchboardListener/Private/SwitchboardListener.cpp
@@ -1295,10 +1295,10 @@ static void FillOutSyncTopologies(TArray<FSyncTopo>& SyncTopos)
 				SyncDisplay.SyncState = TEXT("Unsynced");
 				break;
 			case NVAPI_GSYNC_DISPLAY_SYNC_STATE_SLAVE:
-				SyncDisplay.SyncState = TEXT("Slave");
+				SyncDisplay.SyncState = TEXT("Follower");
 				break;
 			case NVAPI_GSYNC_DISPLAY_SYNC_STATE_MASTER:
-				SyncDisplay.SyncState = TEXT("Master");
+				SyncDisplay.SyncState = TEXT("Leader");
 				break;
 			default:
 				SyncDisplay.SyncState = TEXT("Unknown");
@@ -1344,7 +1344,7 @@ static void FillOutSyncTopologies(TArray<FSyncTopo>& SyncTopos)
 			SyncTopo.SyncStatusParams.RefreshRate = GSyncStatusParams.refreshRate;
 			SyncTopo.SyncStatusParams.HouseSyncIncoming = GSyncStatusParams.houseSyncIncoming;
 			SyncTopo.SyncStatusParams.bHouseSync = !!GSyncStatusParams.bHouseSync;
-			SyncTopo.SyncStatusParams.bInternalSlave = GSyncStatusParams.bInternalSlave;
+			SyncTopo.SyncStatusParams.bInternalSecondary = GSyncStatusParams.bInternalSlave;
 		}
 
 		// Sync Control Parameters
diff --git a/Engine/Source/Programs/SwitchboardListener/Private/SwitchboardListenerVersion.h b/Engine/Source/Programs/SwitchboardListener/Private/SwitchboardListenerVersion.h
--- a/Engine/Source/Programs/SwitchboardListener/Private/SwitchboardListenerVersion.h
+++ b/Engine/Source/Programs/SwitchboardListener/Private/SwitchboardListenerVersion.h
@@ -3,5 +3,5 @@
 #pragma once
 
 #define SBLISTENER_VERSION_MAJOR 2
-#define SBLISTENER_VERSION_MINOR 0
-#define SBLISTENER_VERSION_PATCH 1
+#define SBLISTENER_VERSION_MINOR 1
+#define SBLISTENER_VERSION_PATCH 0
diff --git a/Engine/Source/Programs/SwitchboardListener/Private/SyncStatus.h b/Engine/Source/Programs/SwitchboardListener/Private/SyncStatus.h
--- a/Engine/Source/Programs/SwitchboardListener/Private/SyncStatus.h
+++ b/Engine/Source/Programs/SwitchboardListener/Private/SyncStatus.h
@@ -53,8 +53,9 @@ struct FSyncStatusParams
 	UPROPERTY()
 	bool bHouseSync;
 
+	/** Per nvapi.h comment: "means that this P2061 board receives input from another P2061 board" */
 	UPROPERTY()
-	bool bInternalSlave;
+	bool bInternalSecondary;
 };
 
 USTRUCT()
