SwitchboardEditor: Don't hide toolbar button automatically after creating shortcuts.

diff --git a/Engine/Plugins/VirtualProduction/Switchboard/Source/SwitchboardEditor/Private/SwitchboardEditorModule.cpp b/Engine/Plugins/VirtualProduction/Switchboard/Source/SwitchboardEditor/Private/SwitchboardEditorModule.cpp
--- a/Engine/Plugins/VirtualProduction/Switchboard/Source/SwitchboardEditor/Private/SwitchboardEditorModule.cpp
+++ b/Engine/Plugins/VirtualProduction/Switchboard/Source/SwitchboardEditor/Private/SwitchboardEditorModule.cpp
@@ -349,9 +349,12 @@ TSharedFuture<FSwitchboardVerifyResult> FSwitchboardEditorModule::GetVerifyResul
 			UE_LOG(LogSwitchboardPlugin, Log, TEXT("Verify summary: %d"), Result.Summary);
 			UE_LOG(LogSwitchboardPlugin, Log, TEXT("Verify log: %s"), *Result.Log);
 
+			// Disabled hiding for now; new dropdown menu items have been added which have no other UI,
+			// and some user confusion has been reported. We need to revisit this UX.
+#if 0 // #if SWITCHBOARD_SHORTCUTS
 			// On platforms where we support creating shortcuts, once shortcuts have been created,
 			// we hide our toolbar button to yield space.
-#if SWITCHBOARD_SHORTCUTS
+
 			// This task gets us from the future back onto the main thread prior to manipulating the UI.
 			Async(EAsyncExecution::TaskGraphMainThread, []() {
 				if (FSwitchboardEditorModule::Get().GetSwitchboardInstallState() == ESwitchboardInstallState::Nominal)
