Providing Precompiled Binary syncing support via UnrealGameSync for switchboard.

diff --git a/Engine/Plugins/VirtualProduction/Switchboard/Source/Switchboard/sbl_helper.py b/Engine/Plugins/VirtualProduction/Switchboard/Source/Switchboard/sbl_helper.py
--- a/Engine/Plugins/VirtualProduction/Switchboard/Source/Switchboard/sbl_helper.py
+++ b/Engine/Plugins/VirtualProduction/Switchboard/Source/Switchboard/sbl_helper.py
@@ -11,6 +11,7 @@ import logging
 import marshal
 import os
 import pathlib
+from platform import platform
 from queue import Empty, Queue
 import subprocess
 import sys
@@ -18,6 +19,8 @@ import threading
 import tempfile
 import traceback
 from typing import Callable, Dict, List, Optional, Sequence, Tuple, Union
+from shutil import which
+import re
 
 
 MAX_PARALLEL_WORKERS = 8  # Consistent with UGS, ushell, etc.
@@ -613,6 +616,142 @@ def dispatch_sync_workers(
 
     return 0
 
+def setup_ugs_dependencies(
+    engine_dir: pathlib.Path
+) -> int:
+    logging.info('Verifying UnrealGameSync dependencies.')
+    
+    if sys.platform.startswith('win'):
+        script_name = 'GetDotnetPath.bat'
+        dotnet_setup_script = os.path.join(engine_dir, 'Build', 'BatchFiles', script_name)
+        dotnet_setup_args = [dotnet_setup_script]
+    else:
+        platform_dirname = 'Mac' if sys.platform.startswith('darwin') else 'Linux'
+        platform_scripts_dir = os.path.join(engine_dir, 'Build', 'BatchFiles', f'{platform_dirname}')
+
+        script_name = 'SetupEnvironment.sh'
+        dotnet_setup_script = os.path.join(platform_scripts_dir, script_name)
+        dotnet_setup_args = [dotnet_setup_script, '-dotnet', f'{platform_scripts_dir}']
+
+    try:
+        with subprocess.Popen(dotnet_setup_args, stdin=subprocess.DEVNULL, stdout=subprocess.PIPE) as proc:
+            for line in proc.stdout:
+                logging.debug(f'{script_name}> {line.decode().rstrip()}')
+    except Exception as exc:
+        logging.error(f'setup_ugs_dependencies(): exception running Popen: {dotnet_setup_args}', exc_info=exc)
+
+    if proc.returncode != 0:
+        logging.error('Unable to find a install of Dotnet SDK. Please make sure you have it installed and that `dotnet` is a globally available command.')
+    
+    return proc.returncode
+
+def ugs_sync(
+    ugs_bin_path: pathlib.Path,
+    uproj_path: pathlib.Path,
+    engine_dir: pathlib.Path,
+    sync_cl: Optional[int] = None,
+    user: Optional[str] = None,
+    client: Optional[str] = None,
+    sync_pcbs: Optional[bool] = False
+) -> Optional[int]:
+
+    ## 
+    # Assumption: the current working directory is already set to the sync target
+
+    # Make sure the needed dependencies (dotnet, etc.) are installed/setup
+    deps_result = setup_ugs_dependencies(engine_dir)
+    if deps_result != 0:
+        logging.error('Failed to setup UnrealGameSync dependencies. Aborting sync.')
+        return deps_result
+
+    logging.info("Capturing UnrealGameSync's current state.")
+
+    # Capture the current stat of UGS so we can restore it once done
+    ugs_state_to_restore = {}
+    try:
+        status_args = ['dotnet', ugs_bin_path, 'status']
+        with subprocess.Popen(status_args, stdin=subprocess.DEVNULL, stdout=subprocess.PIPE) as status_proc:
+            for staus_line in status_proc.stdout:
+                staus_line_str =  staus_line.decode().rstrip()
+                logging.debug(f'ugs> {staus_line_str}')
+                if 'Project:' in staus_line_str:
+                        match = re.match("Project:\s*//(\S+?)/(\S+)", staus_line_str)
+                        if match:
+                            ugs_state_to_restore['client'] = match.group(1)
+                            ugs_state_to_restore['project'] = match.group(2)
+                if 'User:' in staus_line_str:
+                        match = re.match("User:\s*(\S+)", staus_line_str)
+                        if match:
+                            ugs_state_to_restore['user'] = match.group(1)
+    except Exception as exc:
+        logging.error(f'ugs_sync(): exception running Popen: {status_args}', exc_info=exc)
+        return -1
+
+    logging.info('Initializing UnrealGameSync...')
+
+    try:
+        # UGS expects a project path relative to the repo's root
+        relative_proj_path = os.path.relpath(uproj_path, engine_dir.parent)
+        # UGS only accepts paths with delimited by forward-slashes
+        norm_proj_path = relative_proj_path.replace('\\','/')
+
+        # Initialize UGS to the project/workspace we want to sync
+        init_args = ['dotnet', ugs_bin_path, 'init', f'-project={norm_proj_path}']
+        if client:
+            init_args.append(f'-client={client}')
+        if user:
+            init_args.append(f'-user={user}')
+
+        with subprocess.Popen(init_args, stdin=subprocess.DEVNULL, stdout=subprocess.PIPE) as init_proc:
+            for init_line in init_proc.stdout:
+                logging.info(f'ugs> {init_line.decode().rstrip()}')
+    except Exception as exc:
+        logging.error(f'ugs_sync(): exception running Popen: {init_args}', exc_info=exc)
+        return -1
+            
+    logging.info('Syncing via UnrealGameSync...')
+    try:
+        # Initialize UGS to the project/workspace we want to sync
+        sync_args = ['dotnet', ugs_bin_path, 'sync']
+        if sync_cl:
+            sync_args.append(f'{sync_cl}')
+        else:
+            sync_args.append('latest')
+
+        if sync_pcbs:
+            sync_args.append('-binaries')
+
+        with subprocess.Popen(sync_args, stdin=subprocess.DEVNULL, stdout=subprocess.PIPE) as sync_proc:
+            for sync_line in sync_proc.stdout:
+                sync_str = sync_line.decode().rstrip()
+                
+                # 'Invalid argument' errors are reported via stdout
+                if sync_pcbs and 'Invalid argument: -binaries' in sync_str:
+                    logging.error(f'ugs> {sync_str}')
+                    logging.error('Your version of UnrealGameSync does not support syncing precompiled binaries via commandline. Please make sure UGS is up to date.')
+                    sync_proc.terminate()
+                else:
+                    logging.info(f'ugs> {sync_str}')      
+    except Exception as exc:
+        logging.error(f'ugs_sync(): exception running Popen: {sync_args}', exc_info=exc)
+        return -1
+
+    if ugs_state_to_restore:
+        logging.info('Restoring UnrealGameSync state.')
+        try:
+            restore_args = ['dotnet', ugs_bin_path, 'init']
+            for key,val in ugs_state_to_restore.items():
+                restore_args.append(f'-{key}={val}')
+
+            with subprocess.Popen(restore_args, stdin=subprocess.DEVNULL, stdout=subprocess.PIPE) as restore_proc:
+                for restore_line in restore_proc.stdout:
+                    logging.debug(f'ugs> {restore_line.decode().rstrip()}')
+        except Exception as exc:
+            logging.error(f'ugs_sync(): exception running Popen: {restore_args}', exc_info=exc)
+            return -1
+
+    return sync_proc.returncode
+
 
 def path_is_relative_to(
     root: Union[str, os.PathLike[str]],
@@ -677,6 +816,12 @@ class SbListenerHelper:
         sync_parser.add_argument(
             '--generate', action='store_true',
             help='Run GenerateProjectFiles after syncing')
+        sync_parser.add_argument(
+            '--use-pcbs', action='store_true',
+            help="Specifies that precompiled binaries should be sync'd along with the project (requires syncing via UnrealGameSync).")
+        sync_parser.add_argument(
+            '--ugs-lib-dir', type=pathlib.Path,
+            help='Directory path specifying where to find the UnrealGameSync library (ugs.dll).')
         sync_parser.add_argument(
             '--clobber-engine', action='store_true',
             help='Override noclobber for engine files')
@@ -767,16 +912,29 @@ class SbListenerHelper:
                                   f'clobber file: {depot_path}')
                     continue
 
-        sync_result = p4_sync(
-            pathspecs,
-            sync_filter_fn=sync_filter,
-            clobber_filter_fn=clobber_filter,
-            progress_fn=self.on_sync_progress,
-            dry_run=self.dry_run,
-            user=self.p4user,
-            client=self.p4client,
-            write_workload_file=options.dump_sync,
-        )
+        sync_result = None
+        if self.use_ugs:
+            logging.info('Syncing via UGS.')
+            sync_result = ugs_sync(
+                self.ugs_bin_path,
+                self.uproj_path,
+                self.find_engine_dir(options),
+                sync_cl=self.sync_project_cl,
+                user=self.p4user,
+                client=self.p4client,
+                sync_pcbs=self.sync_pcbs
+            )
+        else:
+            sync_result = p4_sync(
+                pathspecs,
+                sync_filter_fn=sync_filter,
+                clobber_filter_fn=clobber_filter,
+                progress_fn=self.on_sync_progress,
+                dry_run=self.dry_run,
+                user=self.p4user,
+                client=self.p4client,
+                write_workload_file=options.dump_sync,
+            )
 
         if sync_result is None:
             return 0  # Nothing was synced
@@ -939,6 +1097,7 @@ class SbListenerHelper:
         self.p4client = options.p4client
         self.clobber_engine = options.clobber_engine
         self.clobber_project = options.clobber_project
+        self.ugs_bin_path = self.find_unreal_game_sync_bin(options)
 
         if options.project is None:
             if options.project_cl is not None:
@@ -954,6 +1113,21 @@ class SbListenerHelper:
         self.sync_engine_cl = options.engine_cl
         self.sync_project_cl = options.project_cl
         self.generate_after_sync = options.generate
+        self.sync_pcbs = options.use_pcbs
+        
+        if self.sync_pcbs:
+            if self.sync_engine_cl:
+                self.parser.error('`--use-pcbs` was specified along with `--engine-cl`; these arguments are mutally exclusive and should not be used together.')
+
+            if not self.ugs_bin_path:
+                self.parser.error("Failed to find the UnrealGameSync library (ugs.dll), which `--use-pcbs` requires." 
+                                  "Make sure UGS is installed and either `ugs` is a globally available command, or specify the library's path via the `--ugs-lib-dir` arg.")
+            elif not os.path.exists(self.ugs_bin_path):
+                self.parser.error(f"Failed to find '{self.ugs_bin_path}', which `--use-pcbs` requires." 
+                                  "Make sure UGS is installed and either `ugs` is a globally available command, or specify the library's path via the `--ugs-lib-dir` arg.")
+            else:
+                logging.debug(f'Found UnrealGameSync: {self.ugs_bin_path}')
+        self.use_ugs = self.sync_pcbs
 
         if options.project:
             try:
@@ -1002,6 +1176,27 @@ class SbListenerHelper:
 
         return None
 
+    def find_unreal_game_sync_bin(
+        self,
+        options: argparse.Namespace
+    ) -> Optional[pathlib.Path]:
+        if options.ugs_lib_dir:
+            return os.path.join(os.path.expandvars(options.ugs_lib_dir), 'ugs.dll')
+
+        ugs_env_path = which('ugs')
+        if ugs_env_path:
+            # Since we run via the 'dotnet' command (more platform agnostic), we need the .dll not the executable
+            ugs_path_obj = pathlib.Path(ugs_env_path)
+            return os.path.join(ugs_path_obj.parent, 'ugs.dll')
+
+        if sys.platform.startswith('win'):
+            default_ugs_install_path = os.path.join('${LOCALAPPDATA}', 'UnrealGameSync', 'Latest', 'ugs.dll')
+            return os.path.expandvars(default_ugs_install_path)
+        #elif: for other platforms we don't have a default install location (the dll is usually installed from perforce, which could go anywhere)
+
+        return None
+
+
     def generate_project_files(self) -> int:
         '''
         Invokes GenerateProjectFiles script, streams the output to INFO
diff --git a/Engine/Plugins/VirtualProduction/Switchboard/Source/Switchboard/switchboard/config.py b/Engine/Plugins/VirtualProduction/Switchboard/Source/Switchboard/switchboard/config.py
--- a/Engine/Plugins/VirtualProduction/Switchboard/Source/Switchboard/switchboard/config.py
+++ b/Engine/Plugins/VirtualProduction/Switchboard/Source/Switchboard/switchboard/config.py
@@ -10,6 +10,7 @@ import socket
 import sys
 import typing
 from typing import Type
+from enum import Enum
 
 from PySide2 import QtCore
 from PySide2 import QtGui
@@ -1867,6 +1868,10 @@ class ConfigPathValidator(QtGui.QValidator):
 
         return QtGui.QValidator.Acceptable
 
+class EngineSyncMethod(Enum):
+    Use_Existing = "Use Existing (do not sync/build)"
+    Build_Engine = "Build Engine"
+    Sync_PCBs = "Sync Precompiled Binaries (requires UnrealGameSync)"
 
 class Config(object):
 
@@ -2043,11 +2048,11 @@ class Config(object):
                 data.get('engine_dir', ''),
                 tool_tip="Path to UE 'Engine' directory"
             ),
-            "build_engine": BoolSetting(
-                "build_engine",
-                "Build Engine",
-                data.get('build_engine', False),
-                tool_tip="Is Engine built from source?"
+            'engine_sync_method': OptionSetting(
+                "engine_sync_method",
+                "Engine Sync Method",
+                EngineSyncMethod.Use_Existing.value,
+                possible_values=[p.value for p in EngineSyncMethod],
             ),
             "maps_path": StringSetting(
                 "maps_path",
@@ -2072,10 +2077,17 @@ class Config(object):
             ),
         }
 
+        # Done here outside of the setting's initializor so that values different from the default are flagged in the UI (with a 'reset' option) 
+        if 'engine_sync_method' in data:
+            self.basic_project_settings["engine_sync_method"].update_value(data.get('engine_sync_method'))
+        # To support backwards compatibility, we check the old 'build_engine' option (from before we had a 'PCB' option)
+        elif data.get('build_engine', False):
+            self.basic_project_settings["engine_sync_method"].update_value(EngineSyncMethod.Build_Engine.value)
+
         self.PROJECT_NAME = self.basic_project_settings["project_name"]
         self.UPROJECT_PATH = self.basic_project_settings["uproject"]
         self.ENGINE_DIR = self.basic_project_settings["engine_dir"]
-        self.BUILD_ENGINE = self.basic_project_settings["build_engine"]
+        self.ENGINE_SYNC_METHOD = self.basic_project_settings["engine_sync_method"]
         self.MAPS_PATH = self.basic_project_settings["maps_path"]
         self.MAPS_FILTER = self.basic_project_settings["maps_filter"]
         self.MAPS_PLUGIN_FILTERS = self.basic_project_settings["maps_plugin_filters"]
@@ -2357,7 +2369,7 @@ class Config(object):
         data['project_name'] = self.PROJECT_NAME.get_value()
         data['uproject'] = self.UPROJECT_PATH.get_value()
         data['engine_dir'] = self.ENGINE_DIR.get_value()
-        data['build_engine'] = self.BUILD_ENGINE.get_value()
+        data['engine_sync_method'] = self.ENGINE_SYNC_METHOD.get_value()
         data["maps_path"] = self.MAPS_PATH.get_value()
         data["maps_filter"] = self.MAPS_FILTER.get_value()
         data["maps_plugin_filters"] = self.MAPS_PLUGIN_FILTERS.get_value()
diff --git a/Engine/Plugins/VirtualProduction/Switchboard/Source/Switchboard/switchboard/devices/unreal/plugin_unreal.py b/Engine/Plugins/VirtualProduction/Switchboard/Source/Switchboard/switchboard/devices/unreal/plugin_unreal.py
--- a/Engine/Plugins/VirtualProduction/Switchboard/Source/Switchboard/switchboard/devices/unreal/plugin_unreal.py
+++ b/Engine/Plugins/VirtualProduction/Switchboard/Source/Switchboard/switchboard/devices/unreal/plugin_unreal.py
@@ -27,7 +27,8 @@ from switchboard import switchboard_application
 from switchboard import switchboard_utils as sb_utils
 from switchboard.config import CONFIG, BoolSetting, DirectoryPathSetting, \
     FilePathSetting, IntSetting, MultiOptionSetting, OptionSetting, StringSetting, \
-    SETTINGS, DEFAULT_MAP_TEXT, StringListSetting, migrate_comma_separated_string_to_list
+    SETTINGS, DEFAULT_MAP_TEXT, StringListSetting, migrate_comma_separated_string_to_list, \
+    EngineSyncMethod
 from switchboard.devices.device_base import Device, DeviceStatus, \
     PluginHeaderWidgets
 from switchboard.devices.device_widget_base import DeviceWidget, DeviceAutoJoinMUServerUI
@@ -413,6 +414,16 @@ class DeviceUnreal(Device):
                 'between network messages before considering the connection '
                 'to Switchboard lost and closing it with a timeout error.')
         ),
+        'unrealgamesync_lib_dir': DirectoryPathSetting(
+            attr_name='unrealgamesync_lib_dir',
+            nice_name='UnrealGameSync Dll Dir',
+            value="",
+            tool_tip=(
+                "Expects an absolute path to the device directory containing the ugs.dll library file. \n"
+                "If left blank, switchboard will attempt to find UGS by searching the system PATH and default install locations. \n"
+                "On Windows, the default install location is '${LOCALAPPDATA}/UnrealGameSync/Latest/'.")
+        )
+        
     }
 
     unreal_started_signal = QtCore.Signal()
@@ -554,8 +565,8 @@ class DeviceUnreal(Device):
             self.on_setting_address_changed)
         DeviceUnreal.csettings['port'].signal_setting_changed.connect(
             self.on_setting_port_changed)
-        CONFIG.BUILD_ENGINE.signal_setting_changed.connect(
-            self.on_build_engine_changed)
+        CONFIG.ENGINE_SYNC_METHOD.signal_setting_changed.connect(
+            self.on_engine_sync_method_changed)
 
         self.auto_connect = False
 
@@ -749,6 +760,7 @@ class DeviceUnreal(Device):
             DeviceUnreal.csettings['auto_decline_package_recovery'],
             DeviceUnreal.csettings['udpmessaging_unicast_endpoint'],
             DeviceUnreal.csettings['udpmessaging_extra_static_endpoints'],
+            DeviceUnreal.csettings['unrealgamesync_lib_dir'],
             CONFIG.ENGINE_DIR,
             CONFIG.SOURCE_CONTROL_WORKSPACE,
             CONFIG.UPROJECT_PATH,
@@ -836,8 +848,8 @@ class DeviceUnreal(Device):
             LOGGER.info(f"Updating port for ListenerClient to {new_port}")
             self.unreal_client.port = new_port
 
-    def on_build_engine_changed(self, _, build_engine):
-        if build_engine:
+    def on_engine_sync_method_changed(self, _, engine_sync_method):
+        if engine_sync_method == EngineSyncMethod.Build_Engine.value:
             self.widget.engine_changelist_label.show()
             if not self.is_disconnected:
                 self._request_engine_changelist_number()
@@ -1033,7 +1045,8 @@ class DeviceUnreal(Device):
         project_path = CONFIG.UPROJECT_PATH.get_value(self.name)
         engine_dir = CONFIG.ENGINE_DIR.get_value(self.name)
         workspace = CONFIG.SOURCE_CONTROL_WORKSPACE.get_value(self.name)
-        build_engine = CONFIG.BUILD_ENGINE.get_value()
+        build_engine = CONFIG.ENGINE_SYNC_METHOD.get_value() == EngineSyncMethod.Build_Engine.value
+        sync_precompiled_bins = CONFIG.ENGINE_SYNC_METHOD.get_value() == EngineSyncMethod.Sync_PCBs.value
 
         project_name = os.path.basename(os.path.dirname(project_path))
         LOGGER.info(
@@ -1075,6 +1088,11 @@ class DeviceUnreal(Device):
         if build_engine:
             sync_args += ' --generate'
 
+        if sync_precompiled_bins:
+            sync_args += ' --use-pcbs'
+            if self.unrealgamesync_lib_dir_setting:
+                sync_args += f' --ugs-lib-dir={self.unrealgamesync_lib_dir_setting}'
+
         puuid, msg = message_protocol.create_start_process_message(
             prog_path=sync_tool,
             prog_args=sync_args,
@@ -1136,7 +1154,7 @@ class DeviceUnreal(Device):
         # TODO: Corner case if multiple local devices, and we build on a single
         # local device other than chosen. Guarantee on any local single build?
         if (self.is_designated_local_builder() and
-                CONFIG.BUILD_ENGINE.get_value()):
+                CONFIG.ENGINE_SYNC_METHOD.get_value() == EngineSyncMethod.Build_Engine.value):
             # Build multi-user server
             if CONFIG.MUSERVER_AUTO_BUILD.get_value():
                 if DeviceUnreal.mu_server.is_running():
@@ -1302,6 +1320,11 @@ class DeviceUnreal(Device):
         return DeviceUnreal.csettings[
             'udpmessaging_extra_static_endpoints'].get_value(self.name)
 
+    @property
+    def unrealgamesync_lib_dir_setting(self) -> str:
+        return DeviceUnreal.csettings[
+            'unrealgamesync_lib_dir'].get_value(self.name)
+
     def generate_unreal_exe_path(self):
         return CONFIG.engine_exe_path(
             CONFIG.ENGINE_DIR.get_value(self.name), self.executable_filename)
@@ -1703,7 +1726,7 @@ class DeviceUnreal(Device):
                 self.device_qt_handler.signal_device_sync_failed.emit(self)
 
             # If you build and sync the engine, update its CL
-            if CONFIG.BUILD_ENGINE.get_value():
+            if CONFIG.ENGINE_SYNC_METHOD.get_value() == EngineSyncMethod.Build_Engine.value:
                 self.engine_changelist = (
                     self.inflight_engine_cl
                     if self.inflight_engine_cl is not None
@@ -1738,7 +1761,7 @@ class DeviceUnreal(Device):
                 self._request_project_changelist_number()
                 # Forces an update to the changelist field (to hide the
                 # Building state).
-                if CONFIG.BUILD_ENGINE.get_value():
+                if CONFIG.ENGINE_SYNC_METHOD.get_value() == EngineSyncMethod.Build_Engine.value:
                     self._request_engine_changelist_number()
                     self._request_unreal_editor_version_file()
 
@@ -1946,7 +1969,7 @@ class DeviceUnreal(Device):
         self._request_roles_file()
         self._request_project_changelist_number()
 
-        if CONFIG.BUILD_ENGINE.get_value():
+        if CONFIG.ENGINE_SYNC_METHOD.get_value() == EngineSyncMethod.Build_Engine.value:
             self._request_engine_changelist_number()
             self._request_unreal_editor_version_file()
 
@@ -2453,7 +2476,7 @@ class DeviceWidgetUnreal(DeviceWidget):
         self.project_changelist_label.show()
 
     def update_engine_changelist(self, required_cl: str, synched_cl: str, built__cl: str):
-        if not CONFIG.BUILD_ENGINE.get_value():
+        if not CONFIG.ENGINE_SYNC_METHOD.get_value() == EngineSyncMethod.Build_Engine.value:
             return
         
         self.engine_changelist_label.setText(f'E: {synched_cl}')
@@ -2486,7 +2509,7 @@ class DeviceWidgetUnreal(DeviceWidget):
         sb_widgets.set_qt_property(self.sync_button, 'not_synched', needs_resync)
             
     def update_build_info(self, synched_cl: str, built_cl: str):
-        if built_cl is not None and synched_cl is not None and CONFIG.BUILD_ENGINE.get_value():
+        if built_cl is not None and synched_cl is not None and CONFIG.ENGINE_SYNC_METHOD.get_value() == EngineSyncMethod.Build_Engine.value:
             try:
                 earlier_cl = min(int(built_cl), int(synched_cl))
                 later_cl = max(int(built_cl), int(synched_cl))
diff --git a/Engine/Plugins/VirtualProduction/Switchboard/Source/Switchboard/switchboard/switchboard_dialog.py b/Engine/Plugins/VirtualProduction/Switchboard/Source/Switchboard/switchboard/switchboard_dialog.py
--- a/Engine/Plugins/VirtualProduction/Switchboard/Source/Switchboard/switchboard/switchboard_dialog.py
+++ b/Engine/Plugins/VirtualProduction/Switchboard/Source/Switchboard/switchboard/switchboard_dialog.py
@@ -27,7 +27,7 @@ from switchboard import switchboard_application
 from switchboard import switchboard_utils
 from switchboard import switchboard_widgets as sb_widgets
 from switchboard.add_config_dialog import AddConfigDialog
-from switchboard.config import CONFIG, DEFAULT_MAP_TEXT, SETTINGS
+from switchboard.config import CONFIG, DEFAULT_MAP_TEXT, SETTINGS, EngineSyncMethod
 from switchboard.device_list_widget import DeviceListWidget, DeviceWidgetHeader
 from switchboard.devices.device_base import DeviceStatus
 from switchboard.devices.device_manager import DeviceManager
@@ -660,7 +660,7 @@ class SwitchboardDialog(QtCore.QObject):
     def set_config_hooks(self):
         CONFIG.P4_PROJECT_PATH.signal_setting_changed.connect(lambda: self.p4_refresh_project_cl())
         CONFIG.P4_ENGINE_PATH.signal_setting_changed.connect(lambda: self.p4_refresh_engine_cl())
-        CONFIG.BUILD_ENGINE.signal_setting_changed.connect(lambda: self.p4_refresh_engine_cl())
+        CONFIG.ENGINE_SYNC_METHOD.signal_setting_changed.connect(lambda: self.p4_refresh_engine_cl())
         CONFIG.P4_ENABLED.signal_setting_changed.connect(lambda _, enabled: self.toggle_p4_controls(enabled))
         CONFIG.MAPS_PATH.signal_setting_changed.connect(lambda: self.refresh_levels())
         CONFIG.MAPS_FILTER.signal_setting_changed.connect(lambda: self.refresh_levels())
@@ -1777,7 +1777,7 @@ class SwitchboardDialog(QtCore.QObject):
             return
         self.window.engine_cl_combo_box.clear()
         # if engine is built from source, refresh the engine cl dropdown
-        if CONFIG.BUILD_ENGINE.get_value():
+        if CONFIG.ENGINE_SYNC_METHOD.get_value() == EngineSyncMethod.Build_Engine.value:
             LOGGER.info("Refreshing p4 engine changelists")
             self.window.engine_cl_label.setEnabled(True)
             self.window.engine_cl_combo_box.setEnabled(True)
