Moving ugs utilities into a helper file inside of swichboard (so we can use them in more places).

Leveraging for use when refreshing the project CLs list -- note UGS is currently slower in this op, pausing SB for a moment.

diff --git a/Engine/Plugins/VirtualProduction/Switchboard/Source/Switchboard/sbl_helper.py b/Engine/Plugins/VirtualProduction/Switchboard/Source/Switchboard/sbl_helper.py
--- a/Engine/Plugins/VirtualProduction/Switchboard/Source/Switchboard/sbl_helper.py
+++ b/Engine/Plugins/VirtualProduction/Switchboard/Source/Switchboard/sbl_helper.py
@@ -19,8 +19,7 @@ import threading
 import tempfile
 import traceback
 from typing import Callable, Dict, List, Optional, Sequence, Tuple, Union
-from shutil import which
-import re
+from switchboard import ugs_utils as UGS
 
 
 MAX_PARALLEL_WORKERS = 8  # Consistent with UGS, ushell, etc.
@@ -616,142 +615,6 @@ def dispatch_sync_workers(
 
     return 0
 
-def setup_ugs_dependencies(
-    engine_dir: pathlib.Path
-) -> int:
-    logging.info('Verifying UnrealGameSync dependencies.')
-    
-    if sys.platform.startswith('win'):
-        script_name = 'GetDotnetPath.bat'
-        dotnet_setup_script = os.path.join(engine_dir, 'Build', 'BatchFiles', script_name)
-        dotnet_setup_args = [dotnet_setup_script]
-    else:
-        platform_dirname = 'Mac' if sys.platform.startswith('darwin') else 'Linux'
-        platform_scripts_dir = os.path.join(engine_dir, 'Build', 'BatchFiles', f'{platform_dirname}')
-
-        script_name = 'SetupEnvironment.sh'
-        dotnet_setup_script = os.path.join(platform_scripts_dir, script_name)
-        dotnet_setup_args = [dotnet_setup_script, '-dotnet', f'{platform_scripts_dir}']
-
-    try:
-        with subprocess.Popen(dotnet_setup_args, stdin=subprocess.DEVNULL, stdout=subprocess.PIPE) as proc:
-            for line in proc.stdout:
-                logging.debug(f'{script_name}> {line.decode().rstrip()}')
-    except Exception as exc:
-        logging.error(f'setup_ugs_dependencies(): exception running Popen: {dotnet_setup_args}', exc_info=exc)
-
-    if proc.returncode != 0:
-        logging.error('Unable to find a install of Dotnet SDK. Please make sure you have it installed and that `dotnet` is a globally available command.')
-    
-    return proc.returncode
-
-def ugs_sync(
-    ugs_bin_path: pathlib.Path,
-    uproj_path: pathlib.Path,
-    engine_dir: pathlib.Path,
-    sync_cl: Optional[int] = None,
-    user: Optional[str] = None,
-    client: Optional[str] = None,
-    sync_pcbs: Optional[bool] = False
-) -> Optional[int]:
-
-    ## 
-    # Assumption: the current working directory is already set to the sync target
-
-    # Make sure the needed dependencies (dotnet, etc.) are installed/setup
-    deps_result = setup_ugs_dependencies(engine_dir)
-    if deps_result != 0:
-        logging.error('Failed to setup UnrealGameSync dependencies. Aborting sync.')
-        return deps_result
-
-    logging.info("Capturing UnrealGameSync's current state.")
-
-    # Capture the current stat of UGS so we can restore it once done
-    ugs_state_to_restore = {}
-    try:
-        status_args = ['dotnet', ugs_bin_path, 'status']
-        with subprocess.Popen(status_args, stdin=subprocess.DEVNULL, stdout=subprocess.PIPE) as status_proc:
-            for staus_line in status_proc.stdout:
-                staus_line_str =  staus_line.decode().rstrip()
-                logging.debug(f'ugs> {staus_line_str}')
-                if 'Project:' in staus_line_str:
-                        match = re.match("Project:\s*//(\S+?)/(\S+)", staus_line_str)
-                        if match:
-                            ugs_state_to_restore['client'] = match.group(1)
-                            ugs_state_to_restore['project'] = match.group(2)
-                if 'User:' in staus_line_str:
-                        match = re.match("User:\s*(\S+)", staus_line_str)
-                        if match:
-                            ugs_state_to_restore['user'] = match.group(1)
-    except Exception as exc:
-        logging.error(f'ugs_sync(): exception running Popen: {status_args}', exc_info=exc)
-        return -1
-
-    logging.info('Initializing UnrealGameSync...')
-
-    try:
-        # UGS expects a project path relative to the repo's root
-        relative_proj_path = os.path.relpath(uproj_path, engine_dir.parent)
-        # UGS only accepts paths with delimited by forward-slashes
-        norm_proj_path = relative_proj_path.replace('\\','/')
-
-        # Initialize UGS to the project/workspace we want to sync
-        init_args = ['dotnet', ugs_bin_path, 'init', f'-project={norm_proj_path}']
-        if client:
-            init_args.append(f'-client={client}')
-        if user:
-            init_args.append(f'-user={user}')
-
-        with subprocess.Popen(init_args, stdin=subprocess.DEVNULL, stdout=subprocess.PIPE) as init_proc:
-            for init_line in init_proc.stdout:
-                logging.info(f'ugs> {init_line.decode().rstrip()}')
-    except Exception as exc:
-        logging.error(f'ugs_sync(): exception running Popen: {init_args}', exc_info=exc)
-        return -1
-            
-    logging.info('Syncing via UnrealGameSync...')
-    try:
-        # Initialize UGS to the project/workspace we want to sync
-        sync_args = ['dotnet', ugs_bin_path, 'sync']
-        if sync_cl:
-            sync_args.append(f'{sync_cl}')
-        else:
-            sync_args.append('latest')
-
-        if sync_pcbs:
-            sync_args.append('-binaries')
-
-        with subprocess.Popen(sync_args, stdin=subprocess.DEVNULL, stdout=subprocess.PIPE) as sync_proc:
-            for sync_line in sync_proc.stdout:
-                sync_str = sync_line.decode().rstrip()
-                
-                # 'Invalid argument' errors are reported via stdout
-                if sync_pcbs and 'Invalid argument: -binaries' in sync_str:
-                    logging.error(f'ugs> {sync_str}')
-                    logging.error('Your version of UnrealGameSync does not support syncing precompiled binaries via commandline. Please make sure UGS is up to date.')
-                    sync_proc.terminate()
-                else:
-                    logging.info(f'ugs> {sync_str}')      
-    except Exception as exc:
-        logging.error(f'ugs_sync(): exception running Popen: {sync_args}', exc_info=exc)
-        return -1
-
-    if ugs_state_to_restore:
-        logging.info('Restoring UnrealGameSync state.')
-        try:
-            restore_args = ['dotnet', ugs_bin_path, 'init']
-            for key,val in ugs_state_to_restore.items():
-                restore_args.append(f'-{key}={val}')
-
-            with subprocess.Popen(restore_args, stdin=subprocess.DEVNULL, stdout=subprocess.PIPE) as restore_proc:
-                for restore_line in restore_proc.stdout:
-                    logging.debug(f'ugs> {restore_line.decode().rstrip()}')
-        except Exception as exc:
-            logging.error(f'ugs_sync(): exception running Popen: {restore_args}', exc_info=exc)
-            return -1
-
-    return sync_proc.returncode
-
 
 def path_is_relative_to(
     root: Union[str, os.PathLike[str]],
@@ -914,15 +777,13 @@ class SbListenerHelper:
 
         sync_result = None
         if self.use_ugs:
-            logging.info('Syncing via UGS.')
-            sync_result = ugs_sync(
-                self.ugs_bin_path,
+            sync_result = UGS.sync(
                 self.uproj_path,
-                self.find_engine_dir(options),
                 sync_cl=self.sync_project_cl,
+                sync_pcbs=self.sync_pcbs,
+                ugs_bin_dir=self.ugs_lib_dir,
                 user=self.p4user,
-                client=self.p4client,
-                sync_pcbs=self.sync_pcbs
+                client=self.p4client
             )
         else:
             sync_result = p4_sync(
@@ -1097,7 +958,7 @@ class SbListenerHelper:
         self.p4client = options.p4client
         self.clobber_engine = options.clobber_engine
         self.clobber_project = options.clobber_project
-        self.ugs_bin_path = self.find_unreal_game_sync_bin(options)
+        self.ugs_lib_dir = options.ugs_lib_dir
 
         if options.project is None:
             if options.project_cl is not None:
@@ -1118,15 +979,6 @@ class SbListenerHelper:
         if self.sync_pcbs:
             if self.sync_engine_cl:
                 self.parser.error('`--use-pcbs` was specified along with `--engine-cl`; these arguments are mutally exclusive and should not be used together.')
-
-            if not self.ugs_bin_path:
-                self.parser.error("Failed to find the UnrealGameSync library (ugs.dll), which `--use-pcbs` requires." 
-                                  "Make sure UGS is installed and either `ugs` is a globally available command, or specify the library's path via the `--ugs-lib-dir` arg.")
-            elif not os.path.exists(self.ugs_bin_path):
-                self.parser.error(f"Failed to find '{self.ugs_bin_path}', which `--use-pcbs` requires." 
-                                  "Make sure UGS is installed and either `ugs` is a globally available command, or specify the library's path via the `--ugs-lib-dir` arg.")
-            else:
-                logging.debug(f'Found UnrealGameSync: {self.ugs_bin_path}')
         self.use_ugs = self.sync_pcbs
 
         if options.project:
@@ -1176,27 +1028,6 @@ class SbListenerHelper:
 
         return None
 
-    def find_unreal_game_sync_bin(
-        self,
-        options: argparse.Namespace
-    ) -> Optional[pathlib.Path]:
-        if options.ugs_lib_dir:
-            return os.path.join(os.path.expandvars(options.ugs_lib_dir), 'ugs.dll')
-
-        ugs_env_path = which('ugs')
-        if ugs_env_path:
-            # Since we run via the 'dotnet' command (more platform agnostic), we need the .dll not the executable
-            ugs_path_obj = pathlib.Path(ugs_env_path)
-            return os.path.join(ugs_path_obj.parent, 'ugs.dll')
-
-        if sys.platform.startswith('win'):
-            default_ugs_install_path = os.path.join('${LOCALAPPDATA}', 'UnrealGameSync', 'Latest', 'ugs.dll')
-            return os.path.expandvars(default_ugs_install_path)
-        #elif: for other platforms we don't have a default install location (the dll is usually installed from perforce, which could go anywhere)
-
-        return None
-
-
     def generate_project_files(self) -> int:
         '''
         Invokes GenerateProjectFiles script, streams the output to INFO
diff --git a/Engine/Plugins/VirtualProduction/Switchboard/Source/Switchboard/switchboard/switchboard_dialog.py b/Engine/Plugins/VirtualProduction/Switchboard/Source/Switchboard/switchboard/switchboard_dialog.py
--- a/Engine/Plugins/VirtualProduction/Switchboard/Source/Switchboard/switchboard/switchboard_dialog.py
+++ b/Engine/Plugins/VirtualProduction/Switchboard/Source/Switchboard/switchboard/switchboard_dialog.py
@@ -21,6 +21,7 @@ from PySide2.QtWidgets import QWidgetAction, QMenu
 from switchboard import config
 from switchboard import config_osc as osc
 from switchboard import p4_utils
+from switchboard import ugs_utils
 from switchboard import recording
 from switchboard import resources  # noqa
 from switchboard import switchboard_application
@@ -1762,9 +1763,20 @@ class SwitchboardDialog(QtCore.QObject):
     def p4_refresh_project_cl(self):
         if not CONFIG.P4_ENABLED.get_value():
             return
-        LOGGER.info("Refreshing p4 project changelists")
-        working_dir = os.path.dirname(CONFIG.UPROJECT_PATH.get_value())
-        changelists = p4_utils.p4_latest_changelist(CONFIG.P4_PROJECT_PATH.get_value(), working_dir)
+
+        changelists = None
+        # If we're syncing 'Precompiled Binaries', then that implies that we should be using UGS:
+        if CONFIG.ENGINE_SYNC_METHOD.get_value() == EngineSyncMethod.Sync_PCBs.value:
+            LOGGER.info("Using UnrealGameSync to refresh project changelists.")
+            changelists = ugs_utils.latest_chagelists(Path(CONFIG.UPROJECT_PATH.get_value()), client=CONFIG.SOURCE_CONTROL_WORKSPACE.get_value())
+            if not changelists:
+                LOGGER.error("UnrealGameSync failed to get the project's latest changelists. Falling back to using p4 commands directly.")
+
+        if not changelists:
+            LOGGER.info("Refreshing p4 project changelists")
+            working_dir = os.path.dirname(CONFIG.UPROJECT_PATH.get_value())
+            changelists = p4_utils.p4_latest_changelist(CONFIG.P4_PROJECT_PATH.get_value(), working_dir)
+            
         self.window.project_cl_combo_box.clear()
 
         if changelists:
diff --git a/Engine/Plugins/VirtualProduction/Switchboard/Source/Switchboard/switchboard/ugs_utils.py b/Engine/Plugins/VirtualProduction/Switchboard/Source/Switchboard/switchboard/ugs_utils.py
new file mode 100644
--- /dev/null
+++ b/Engine/Plugins/VirtualProduction/Switchboard/Source/Switchboard/switchboard/ugs_utils.py
@@ -0,0 +1,263 @@
+# Copyright Epic Games, Inc. All Rights Reserved.
+
+import logging
+import subprocess
+import sys
+import shutil
+import re
+from os import path
+from pathlib import Path
+from typing import Optional, List, Dict, Callable
+
+
+def find_bin(
+    dir_to_look_under: Optional[Path] = None
+) -> Optional[Path]:
+    ugs_dir = dir_to_look_under
+
+    if not ugs_dir:
+        env_ugs_path = shutil.which('ugs')
+        if env_ugs_path:
+            ugs_dir = Path(env_ugs_path).parent
+        elif sys.platform.startswith('win'):
+            ugs_dir = path.join('${LOCALAPPDATA}', 'UnrealGameSync', 'Latest')
+        #elif: for other platforms we don't have a default install location (the dll is usually installed from perforce, which could go anywhere)
+
+    if not ugs_dir:
+        logging.error('Could not determine where to find the UnrealGameSync library (ugs.dll). Ensure that UGS is installed on the target device.'
+                      'If it already is installed, then either make sure the `ugs` command is globally available, or specify the directory explicitly.')
+        return None
+
+    ugs_bin_path = path.join(path.expandvars(ugs_dir), 'ugs.dll')
+    if not path.exists(ugs_bin_path):
+        logging.error(f"Failed to find '{ugs_bin_path}'. Ensure that UGS is installed on the target device."
+                       'If it already is installed, then either make sure the `ugs` command is globally available, or specify the directory explicitly.')
+        return None
+
+    return ugs_bin_path
+
+def setup_dependencies(
+    ue_engine_dir:Path
+) -> int:
+    logging.debug('Verifying UnrealGameSync dependencies.')
+    
+    if sys.platform.startswith('win'):
+        script_name = 'GetDotnetPath.bat'
+        dotnet_setup_script = path.join(ue_engine_dir, 'Build', 'BatchFiles', script_name)
+        dotnet_setup_args = [dotnet_setup_script]
+    else:
+        platform_dirname = 'Mac' if sys.platform.startswith('darwin') else 'Linux'
+        platform_scripts_dir = path.join(ue_engine_dir, 'Build', 'BatchFiles', f'{platform_dirname}')
+
+        script_name = 'SetupEnvironment.sh'
+        dotnet_setup_script = path.join(platform_scripts_dir, script_name)
+        dotnet_setup_args = [dotnet_setup_script, '-dotnet', f'{platform_scripts_dir}']
+
+    try:
+        with subprocess.Popen(dotnet_setup_args, stdin=subprocess.DEVNULL, stdout=subprocess.PIPE) as proc:
+            for line in proc.stdout:
+                logging.debug(f'{script_name}> {line.decode().rstrip()}')
+    except Exception as exc:
+        logging.error(f'setup_ugs_dependencies(): exception running Popen: {dotnet_setup_args}', exc_info=exc)
+
+    if proc.returncode != 0:
+        logging.error('Unable to find a install of Dotnet SDK. Please make sure you have it installed and that `dotnet` is a globally available command.')
+    
+    return proc.returncode
+
+def _find_engine_dir(
+    uproj_path: Path
+) -> Optional[Path]:
+    iter_dir = uproj_path.parent
+    while not iter_dir.samefile(iter_dir.anchor):
+        candidate = iter_dir / 'Engine'
+        if candidate.is_dir():
+            return candidate
+        iter_dir = iter_dir.parent
+
+    # TODO: Consider looking for Engine dir as ancestor of this script?
+    return None
+
+def _get_active_ugs_context(
+    ugs_dll_path: Path,
+    cwd: Path = None
+) -> Optional[Dict[str, str]]:
+    # Assumes UGS dependencies have been already setup (this assumption prevents us from redundently running `setup_dependencies()`)
+
+    ugs_state = {}
+    try:
+        status_args = ['dotnet', ugs_dll_path, 'status']
+        with subprocess.Popen(status_args, cwd=cwd, stdin=subprocess.DEVNULL, stdout=subprocess.PIPE) as status_proc:
+            for staus_line in status_proc.stdout:
+                staus_line_str =  staus_line.decode().rstrip()
+                logging.debug(f'ugs> {staus_line_str}')
+                if 'Project:' in staus_line_str:
+                        match = re.match("Project:\s*//(\S+?)/(\S+)", staus_line_str)
+                        if match:
+                            ugs_state['client'] = match.group(1)
+                            ugs_state['project'] = match.group(2)
+                if 'User:' in staus_line_str:
+                        match = re.match("User:\s*(\S+)", staus_line_str)
+                        if match:
+                            ugs_state['user'] = match.group(1)
+    except Exception as exc:
+        logging.error(f'_get_ugs_state(): exception running Popen: {status_args}', exc_info=exc)
+        return None
+
+    if not ugs_state:
+        return None
+    return ugs_state
+
+def _set_active_ugs_context(
+    ugs_dll_path: Path,
+    ugs_settings: Dict[str, str],
+    cwd: Path = None
+) -> int:
+    # Assumes UGS dependencies have been already setup (this assumption prevents us from redundently running `setup_dependencies()`)
+
+    try:
+        init_args = ['dotnet', ugs_dll_path, 'init']
+        for key,val in ugs_settings.items():
+            init_args.append(f'-{key}={val}')
+
+        with subprocess.Popen(init_args, cwd=cwd, stdin=subprocess.DEVNULL, stdout=subprocess.PIPE) as init_proc:
+            for init_line in init_proc.stdout:
+                logging.info(f'ugs> {init_line.decode().rstrip()}')
+    except Exception as exc:
+        logging.error(f'_set_active_ugs_context(): exception running Popen: {init_args}', exc_info=exc)
+        return -1
+
+    return 0
+
+def run(
+    ugs_args: List[str],
+    uproj_path: Path,
+    ugs_bin_dir: Optional[Path] = None,
+    user: Optional[str] = None,
+    client: Optional[str] = None,
+    output_handling_fn: Optional[Callable[[str], int]] = None,
+) -> int:
+
+    # Even though we specify this should be a list, handle the simplified case of users passing a single word command
+    if type(ugs_args) == str:
+        ugs_args = ugs_args.split()
+    args_str = ' '.join(ugs_args)
+
+    ugs_bin_path = find_bin(ugs_bin_dir)
+    if not ugs_bin_path:
+        logging.error(f"Failed to find 'ugs.dll'. Aborting UnrealGameSync command: `ugs {args_str}`")
+        return -1
+
+    ue_engine_dir = _find_engine_dir(uproj_path)
+    if not ue_engine_dir:
+        logging.warning("Failed to locate the Unreal '/Engine/' directory. As a result we are unable to validate UGS dependencies. Ignoring and continuing on as if they were setup.")
+    # Make sure the needed dependencies (dotnet, etc.) are installed upfront
+    elif setup_dependencies(ue_engine_dir) != 0:
+        logging.warning('Failed to validate UnrealGameSync dependencies. Ignoring and continuing on as if they were setup.')
+
+    cwd = uproj_path.parent
+
+    logging.info("Capturing UnrealGameSync's current state.")
+    ugs_state_to_restore = _get_active_ugs_context(ugs_bin_path, cwd=cwd)
+    if not ugs_state_to_restore:
+        logging.warning('Failed to capture the current state of UnrealGameSync. Will not be able to restore it after completing our operation.')
+
+    logging.info('Setting UnrealGameSync context.')
+    # UGS expects a project path relative to the repo's root
+    sanitized_proj_path = path.relpath(uproj_path, ue_engine_dir.parent) if ue_engine_dir else uproj_path
+    # and only accepts paths with delimited by forward-slashes
+    sanitized_proj_path = sanitized_proj_path.replace('\\','/')
+    op_context_params = { 'project' : sanitized_proj_path }
+    if user:
+        op_context_params['user'] = user
+    if client:
+        op_context_params['client'] = client
+
+    init_result = _set_active_ugs_context(ugs_bin_path, op_context_params, cwd=cwd)
+    if init_result != 0:
+        logging.error(f"Failed to initialize UnrealGameSync. Aborting UnrealGameSync command: `ugs {args_str}`")
+        return init_result
+
+    logging.info(f'Executing UnrealGameSync command: `ugs {args_str}`.')
+    ugs_cmd = ['dotnet', ugs_bin_path] + ugs_args
+
+    try:
+        with subprocess.Popen(ugs_cmd, cwd=cwd, stdin=subprocess.DEVNULL, stdout=subprocess.PIPE) as ugs_proc:
+            for line in ugs_proc.stdout:
+                line_str = f'{line.decode().rstrip()}'
+
+                invalid_arg_match = re.match("Invalid argument:\s*(\S+)", line_str)
+                if invalid_arg_match:
+                    logging.error(f'ugs> {line_str}')
+                    logging.error(f'The current version of UnrealGameSync does not support the `{invalid_arg_match.match.group(1)}` argument. Please make sure UGS is up to date.')
+                    ugs_proc.terminate()
+                elif output_handling_fn:
+                    if output_handling_fn(line_str) != 0:
+                        ugs_proc.terminate()
+                else:
+                    logging.info(f'ugs> {line_str}')
+    except Exception as exc:
+        logging.error(f'ugs.run(): exception running Popen: {ugs_cmd}', exc_info=exc)
+        return -1
+
+    if ugs_state_to_restore:
+        logging.info('Restoring previous UnrealGameSync state.')
+        if _set_active_ugs_context(ugs_bin_path, ugs_state_to_restore, cwd=cwd) != 0:
+            logging.warn("Failed to resotre UnrealGameSync's state prior to operation.")
+
+    return ugs_proc.returncode
+
+def sync(
+    uproj_path: Path,
+    sync_cl: Optional[int] = None,
+    sync_pcbs: Optional[bool] = False,
+    ugs_bin_dir: Optional[Path] = None,
+    user: Optional[str] = None,
+    client: Optional[str] = None
+) -> Optional[int]:
+
+    sync_args = ['sync']
+    if sync_cl:
+        sync_args.append(f'{sync_cl}')
+    else:
+        sync_args.append('latest')
+
+    if sync_pcbs:
+        sync_args.append('-binaries')
+
+    return run(
+        sync_args,
+        uproj_path,
+        ugs_bin_dir=ugs_bin_dir,
+        user=user,
+        client=client
+    )
+
+def latest_chagelists(
+    uproj_path: Path,
+    ugs_bin_dir: Optional[Path] = None,
+    user: Optional[str] = None,
+    client: Optional[str] = None
+) -> Optional[List[int]]:
+    
+    changelist_list = []
+
+    def changes_output_handler(output_str:str, changes_out=changelist_list) -> int:
+        cl_desc_match = re.match("^\s*(\d+).*", output_str)
+        if cl_desc_match:
+            changes_out.append(cl_desc_match.group(1))
+        return 0
+
+    run_result = run(
+        'changes',
+        uproj_path,
+        ugs_bin_dir =ugs_bin_dir,
+        user = user,
+        client=client,
+        output_handling_fn = changes_output_handler
+    )
+
+    if run_result != 0:
+        return None
+
+    return changelist_list
\ No newline at end of file
