Switchboard: Bump Listener version to 2.2.0 to reflect recent UGS sync support additions.

diff --git a/Engine/Plugins/VirtualProduction/Switchboard/Source/Switchboard/switchboard/devices/unreal/version_helpers.py b/Engine/Plugins/VirtualProduction/Switchboard/Source/Switchboard/switchboard/devices/unreal/version_helpers.py
--- a/Engine/Plugins/VirtualProduction/Switchboard/Source/Switchboard/switchboard/devices/unreal/version_helpers.py
+++ b/Engine/Plugins/VirtualProduction/Switchboard/Source/Switchboard/switchboard/devices/unreal/version_helpers.py
@@ -3,7 +3,7 @@
 import typing
 
 
-LISTENER_COMPATIBLE_VERSION = (2,1)
+LISTENER_COMPATIBLE_VERSION = (2,2)
 
 # Version where redeploy support was initially introduced.
 LISTENER_MIN_REDEPLOY_VERSION = (1,4,0)
diff --git a/Engine/Source/Programs/SwitchboardListener/Private/SwitchboardListenerVersion.h b/Engine/Source/Programs/SwitchboardListener/Private/SwitchboardListenerVersion.h
--- a/Engine/Source/Programs/SwitchboardListener/Private/SwitchboardListenerVersion.h
+++ b/Engine/Source/Programs/SwitchboardListener/Private/SwitchboardListenerVersion.h
@@ -3,5 +3,5 @@
 #pragma once
 
 #define SBLISTENER_VERSION_MAJOR 2
-#define SBLISTENER_VERSION_MINOR 1
+#define SBLISTENER_VERSION_MINOR 2
 #define SBLISTENER_VERSION_PATCH 0
