Core - Some more build fixes!

diff --git a/Engine/Plugins/VirtualProduction/Switchboard/Source/SwitchboardEditor/Private/SwitchboardEditorModule.cpp b/Engine/Plugins/VirtualProduction/Switchboard/Source/SwitchboardEditor/Private/SwitchboardEditorModule.cpp
--- a/Engine/Plugins/VirtualProduction/Switchboard/Source/SwitchboardEditor/Private/SwitchboardEditorModule.cpp
+++ b/Engine/Plugins/VirtualProduction/Switchboard/Source/SwitchboardEditor/Private/SwitchboardEditorModule.cpp
@@ -224,7 +224,7 @@ bool FSwitchboardEditorModule::CreateNewConfig(const FSwitchboardNewConfigUserOp
 	if (const UWorld* World = GEditor ? GEditor->GetEditorWorldContext(false).World() : nullptr)
 	{
 		FString ClassName, PackageName, ObjectName, SubObjectName;
-		FPackageName::SplitFullObjectPath(World->PersistentLevel.GetPath(), ClassName, PackageName, ObjectName, SubObjectName, true);
+		FPackageName::SplitFullObjectPath(World->PersistentLevel.GetPathName(), ClassName, PackageName, ObjectName, SubObjectName, true);
 		ScriptArgs.Map = PackageName;
 	}
 
diff --git a/Engine/Source/Runtime/CoreUObject/Public/UObject/ObjectPtr.h b/Engine/Source/Runtime/CoreUObject/Public/UObject/ObjectPtr.h
--- a/Engine/Source/Runtime/CoreUObject/Public/UObject/ObjectPtr.h
+++ b/Engine/Source/Runtime/CoreUObject/Public/UObject/ObjectPtr.h
@@ -142,6 +142,10 @@ public:
 
 	FORCEINLINE bool IsResolved() const { return IsObjectHandleResolved(Handle); }
 
+	// Gets the PathName of the object without resolving the object reference.
+	// @TODO OBJPTR: Deprecate this.
+	FString GetPath() const { return GetPathName(); }
+
 #if UE_WITH_OBJECT_HANDLE_LATE_RESOLVE
 	// Gets the PathName of the object without resolving the object reference.
 	COREUOBJECT_API FString GetPathName() const;
diff --git a/Engine/Source/Runtime/CoreUObject/Tests/ObjectRefTrackingTestBase.cpp b/Engine/Source/Runtime/CoreUObject/Tests/ObjectRefTrackingTestBase.cpp
--- a/Engine/Source/Runtime/CoreUObject/Tests/ObjectRefTrackingTestBase.cpp
+++ b/Engine/Source/Runtime/CoreUObject/Tests/ObjectRefTrackingTestBase.cpp
@@ -4,6 +4,11 @@
 
 #include "ObjectRefTrackingTestBase.h"
 
+#if UE_WITH_OBJECT_HANDLE_TRACKING
+FDelegateHandle FObjectRefTrackingTestBase::ResolvedCallbackHandle;
+FDelegateHandle FObjectRefTrackingTestBase::HandleReadCallbackHandle;
+#endif
+
 thread_local uint32 FObjectRefTrackingTestBase::NumResolves = 0;
 thread_local uint32 FObjectRefTrackingTestBase::NumFailedResolves = 0;
 thread_local uint32 FObjectRefTrackingTestBase::NumReads = 0;
