Adding Switchboard option to sync with UGS and not PCBs.

diff --git a/Engine/Plugins/VirtualProduction/Switchboard/Source/Switchboard/sbl_helper.py b/Engine/Plugins/VirtualProduction/Switchboard/Source/Switchboard/sbl_helper.py
--- a/Engine/Plugins/VirtualProduction/Switchboard/Source/Switchboard/sbl_helper.py
+++ b/Engine/Plugins/VirtualProduction/Switchboard/Source/Switchboard/sbl_helper.py
@@ -679,6 +679,9 @@ class SbListenerHelper:
         sync_parser.add_argument(
             '--generate', action='store_true',
             help='Run GenerateProjectFiles after syncing')
+        sync_parser.add_argument(
+            '--use-ugs', action='store_true',
+            help="Specifies that UnrealGameSync should be used to sync the project and engine together.")
         sync_parser.add_argument(
             '--use-pcbs', action='store_true',
             help="Specifies that precompiled binaries should be sync'd along with the project (requires syncing via UnrealGameSync).")
@@ -973,13 +976,15 @@ class SbListenerHelper:
 
         self.sync_engine_cl = options.engine_cl
         self.sync_project_cl = options.project_cl
-        self.generate_after_sync = options.generate
+        self.use_ugs = options.use_ugs
+        self.generate_after_sync = options.generate and not self.use_ugs # UGS will generate the project files by default (no need to duplicate work)
         self.sync_pcbs = options.use_pcbs
+
+        if self.sync_pcbs and not self.use_ugs:
+            self.parser.error('`--use-pcbs` requires `--use-ugs`.')
         
-        if self.sync_pcbs:
-            if self.sync_engine_cl:
-                self.parser.error('`--use-pcbs` was specified along with `--engine-cl`; these arguments are mutally exclusive and should not be used together.')
-        self.use_ugs = self.sync_pcbs
+        if self.use_ugs and self.sync_engine_cl:
+            self.parser.error('`--use-ugs` was specified along with `--engine-cl`; these arguments are mutally exclusive and should not be used together.')
 
         if options.project:
             try:
diff --git a/Engine/Plugins/VirtualProduction/Switchboard/Source/Switchboard/switchboard/config.py b/Engine/Plugins/VirtualProduction/Switchboard/Source/Switchboard/switchboard/config.py
--- a/Engine/Plugins/VirtualProduction/Switchboard/Source/Switchboard/switchboard/config.py
+++ b/Engine/Plugins/VirtualProduction/Switchboard/Source/Switchboard/switchboard/config.py
@@ -1872,6 +1872,7 @@ class EngineSyncMethod(Enum):
     Use_Existing = "Use Existing (do not sync/build)"
     Build_Engine = "Build Engine"
     Sync_PCBs = "Sync Precompiled Binaries (requires UnrealGameSync)"
+    Sync_From_UGS = "Sync engine and project together using UnrealGameSync"
 
 class Config(object):
 
diff --git a/Engine/Plugins/VirtualProduction/Switchboard/Source/Switchboard/switchboard/devices/unreal/plugin_unreal.py b/Engine/Plugins/VirtualProduction/Switchboard/Source/Switchboard/switchboard/devices/unreal/plugin_unreal.py
--- a/Engine/Plugins/VirtualProduction/Switchboard/Source/Switchboard/switchboard/devices/unreal/plugin_unreal.py
+++ b/Engine/Plugins/VirtualProduction/Switchboard/Source/Switchboard/switchboard/devices/unreal/plugin_unreal.py
@@ -1045,8 +1045,11 @@ class DeviceUnreal(Device):
         project_path = CONFIG.UPROJECT_PATH.get_value(self.name)
         engine_dir = CONFIG.ENGINE_DIR.get_value(self.name)
         workspace = CONFIG.SOURCE_CONTROL_WORKSPACE.get_value(self.name)
-        build_engine = CONFIG.ENGINE_SYNC_METHOD.get_value() == EngineSyncMethod.Build_Engine.value
-        sync_precompiled_bins = CONFIG.ENGINE_SYNC_METHOD.get_value() == EngineSyncMethod.Sync_PCBs.value
+
+        sync_method = CONFIG.ENGINE_SYNC_METHOD.get_value()
+        generate_proj_files = sync_method == EngineSyncMethod.Build_Engine.value
+        sync_precompiled_bins = sync_method == EngineSyncMethod.Sync_PCBs.value
+        sync_using_ugs = sync_precompiled_bins or sync_method == EngineSyncMethod.Sync_From_UGS.value
 
         project_name = os.path.basename(os.path.dirname(project_path))
         LOGGER.info(
@@ -1085,14 +1088,17 @@ class DeviceUnreal(Device):
             sync_args += f' --project-cl={project_cl} --clobber-project'
             self.inflight_project_cl = project_cl
 
-        if build_engine:
+        if generate_proj_files:
             sync_args += ' --generate'
 
-        if sync_precompiled_bins:
-            sync_args += ' --use-pcbs'
+        if sync_using_ugs:
+            sync_args += ' --use-ugs'
             if self.unrealgamesync_lib_dir_setting:
                 sync_args += f' --ugs-lib-dir={self.unrealgamesync_lib_dir_setting}'
-
+        
+        if sync_precompiled_bins:
+            sync_args += ' --use-pcbs'
+            
             # If we're syncing 'Precompiled Binaries' one of those binaries may be the SwitchboardListener executable
             # which means (on Windows atleast) we need to move the executable to make way for the new one
             _, msg = message_protocol.create_free_listener_bin_message()
diff --git a/Engine/Plugins/VirtualProduction/Switchboard/Source/Switchboard/switchboard/switchboard_dialog.py b/Engine/Plugins/VirtualProduction/Switchboard/Source/Switchboard/switchboard/switchboard_dialog.py
--- a/Engine/Plugins/VirtualProduction/Switchboard/Source/Switchboard/switchboard/switchboard_dialog.py
+++ b/Engine/Plugins/VirtualProduction/Switchboard/Source/Switchboard/switchboard/switchboard_dialog.py
@@ -1764,9 +1764,11 @@ class SwitchboardDialog(QtCore.QObject):
         if not CONFIG.P4_ENABLED.get_value():
             return
 
+        sync_method = CONFIG.ENGINE_SYNC_METHOD.get_value()
+
         changelists = None
         # If we're syncing 'Precompiled Binaries', then that implies that we should be using UGS:
-        if CONFIG.ENGINE_SYNC_METHOD.get_value() == EngineSyncMethod.Sync_PCBs.value:
+        if sync_method == EngineSyncMethod.Sync_PCBs.value or sync_method == EngineSyncMethod.Sync_From_UGS.value:
             LOGGER.info("Using UnrealGameSync to refresh project changelists.")
             changelists = ugs_utils.latest_chagelists(Path(CONFIG.UPROJECT_PATH.get_value()), client=CONFIG.SOURCE_CONTROL_WORKSPACE.get_value())
             if not changelists:
