Switchboard: generalized mechanism for exposing plugin contents for use in UI

This enables specifying plugins that Switchboard should search for content not just
by plugin name but also by relative or absolute file path to the plugin directory.

This same mechanism is now used by both the level and nDisplay config dropdown
menus. The old level-specific "maps_plugin_filters" config setting was deprecated
in favor of the newly added and more general "content_plugin_filters" setting, and
any "maps_plugin_filters" settings in existing configs will be automatically upgraded.

diff --git a/Engine/Plugins/VirtualProduction/Switchboard/Source/Switchboard/switchboard/config.py b/Engine/Plugins/VirtualProduction/Switchboard/Source/Switchboard/switchboard/config.py
--- a/Engine/Plugins/VirtualProduction/Switchboard/Source/Switchboard/switchboard/config.py
+++ b/Engine/Plugins/VirtualProduction/Switchboard/Source/Switchboard/switchboard/config.py
@@ -18,8 +18,9 @@ from PySide2 import QtWidgets
 
 from switchboard import switchboard_widgets as sb_widgets
 from switchboard.switchboard_logging import LOGGER
-from switchboard.switchboard_widgets import (DropDownMenuComboBox,
-    NonScrollableComboBox)
+from switchboard.switchboard_widgets import (
+    DropDownMenuComboBox, NonScrollableComboBox)
+from switchboard import ue_plugin_utils
 
 ROOT_CONFIGS_PATH = pathlib.Path(__file__).parent.with_name('configs')
 CONFIG_SUFFIX = '.json'
@@ -2073,7 +2074,28 @@ class Config(object):
                 "maps_plugin_filters",
                 "Map Plugin Filters",
                 data.get('maps_plugin_filters', []),
-                tool_tip="Plugins whose name matches any of these filters will also be searched for maps.",
+                tool_tip=(
+                    "DEPRECATED: Use 'Content Plugin Filters' instead' - "
+                    "Plugins whose name matches any of these filters will "
+                    "also be searched for maps."),
+                show_ui=False,
+                migrate_data=migrate_comma_separated_string_to_list
+            ),
+            'content_plugin_filters': StringListSetting(
+                "content_plugin_filters",
+                "Content Plugin Filters",
+                # "Upgrade" from the deprecated "maps_plugin_filters".
+                data.get(
+                    'content_plugin_filters',
+                    data.get('maps_plugin_filters', [])),
+                tool_tip=(
+                    "Plugins that match any of these filters will also be "
+                    "searched when populating fields in Switchboard (e.g. "
+                    "levels, nDisplay configs, etc.). Each value can be "
+                    "either just a plugin name identifying a project plugin, "
+                    "or a relative or absolute path to a plugin directory. "
+                    "Relative paths should be relative to the directory "
+                    "containing the .uproject file."),
                 migrate_data=migrate_comma_separated_string_to_list
             ),
         }
@@ -2091,7 +2113,8 @@ class Config(object):
         self.ENGINE_SYNC_METHOD = self.basic_project_settings["engine_sync_method"]
         self.MAPS_PATH = self.basic_project_settings["maps_path"]
         self.MAPS_FILTER = self.basic_project_settings["maps_filter"]
-        self.MAPS_PLUGIN_FILTERS = self.basic_project_settings["maps_plugin_filters"]
+        self.CONTENT_PLUGIN_FILTERS = (
+            self.basic_project_settings["content_plugin_filters"])
 
         self.osc_settings = {
             "osc_server_port": IntSetting(
@@ -2373,7 +2396,8 @@ class Config(object):
         data['engine_sync_method'] = self.ENGINE_SYNC_METHOD.get_value()
         data["maps_path"] = self.MAPS_PATH.get_value()
         data["maps_filter"] = self.MAPS_FILTER.get_value()
-        data["maps_plugin_filters"] = self.MAPS_PLUGIN_FILTERS.get_value()
+        data["content_plugin_filters"] = (
+            self.CONTENT_PLUGIN_FILTERS.get_value())
         data["listener_exe"] = self.LISTENER_EXE.get_value()
 
         self.save_unreal_insights(data)
@@ -2482,59 +2506,39 @@ class Config(object):
         '''
         return os.path.join(self.get_project_dir(), 'Content')
 
-    def get_project_plugins_dir(self) -> str:
-        '''
-        Get the "Plugins" directory of the project where all of the
-        project-based plugins are located.
-        '''
-        return os.path.join(self.get_project_dir(), 'Plugins')
-
-    def get_project_plugin_names(self) -> typing.List[str]:
-        '''
-        Get a list of the names of all project-based plugins in the project.
-        '''
-        project_plugins_path = os.path.normpath(self.get_project_plugins_dir())
-        if not os.path.isdir(project_plugins_path):
-            return []
-
-        plugin_names = [
-            x.name for x in os.scandir(project_plugins_path) if x.is_dir()]
-
-        return plugin_names
-
-    def get_project_plugin_content_dir(self, plugin_name: str) -> str:
+    def get_unreal_content_plugins(
+            self) -> typing.List[ue_plugin_utils.UnrealPlugin]:
         '''
-        Get the "Content" directory of the project-based plugin with the
-        given name.
-        '''
-        return os.path.join(
-            self.get_project_plugins_dir(), plugin_name, 'Content')
+        Get a list of Unreal Engine plugins that match the current
+        Switchboard config's content plugin filter settings.
 
-    def plugin_name_matches(self, plugin_name: str, name_filters: typing.List[str]) -> bool:
+        By default, Switchboard does not search inside plugins for assets when
+        populating UI fields such as the level or nDisplay config dropdown
+        menus. Plugins in which to search for content can be selectively
+        added though using the "Content Plugin Filters" setting, which
+        stores a list of glob-style patterns.
         '''
-        Test whether a plugin name matches any of the provided filters.
+        filter_patterns = self.CONTENT_PLUGIN_FILTERS.get_value()
+        unreal_plugins = ue_plugin_utils.UnrealPlugin.from_path_filters(
+            self.get_project_dir(), filter_patterns)
+        return unreal_plugins
 
-        Returns True if any one of the filters matches, or False otherwise.
-        '''
-        for name_filter in name_filters:
-            if fnmatch.fnmatch(plugin_name, name_filter):
-                return True
-
-        return False
-
-    def resolve_content_path(self, file_path: str, plugin_name: str = None) -> str:
+    def resolve_content_path(
+            self,
+            file_path: str,
+            unreal_content_plugin: ue_plugin_utils.UnrealPlugin = None) -> str:
         '''
         Resolve a file path on the file system to the corresponding content
         path in UE.
 
-        If a plugin_name is provided, the file is assumed to live inside that
-        plugin and its content path will have the appropriate plugin
-        name-based prefix. Otherwise, the file is assumed to live inside the
-        project's content folder.
+        If an unreal_content_plugin is provided, the file is assumed to live
+        inside that plugin and its content path will have the appropriate
+        plugin name-based prefix. Otherwise, the file is assumed to live
+        inside the project's content folder.
         '''
-        if plugin_name:
-            content_dir = self.get_project_plugin_content_dir(plugin_name)
-            ue_path_prefix = f'/{plugin_name}'
+        if unreal_content_plugin:
+            content_dir = str(unreal_content_plugin.plugin_content_path)
+            ue_path_prefix = str(unreal_content_plugin.mounted_path)
         else:
             content_dir = self.get_project_content_dir()
             ue_path_prefix = '/Game'
@@ -2547,44 +2551,43 @@ class Config(object):
     def maps(self):
         '''
         Returns a list of full map paths in an Unreal Engine project and
-        in project-based plugins such as:
+        in plugins such as:
             [
                 "/Game/Maps/MapName",
                 "/MyPlugin/Levels/MapName"
             ]
         The slashes will always be "/" independent of the platform's separator.
         '''
-        content_maps_path = os.path.normpath(
+        project_maps_path = os.path.normpath(
             os.path.join(
                 self.get_project_content_dir(),
                 self.MAPS_PATH.get_value()))
 
-        # maps_paths stores a list of tuples of the form
-        # (plugin_name, file_path). This allows us to differentiate between
-        # maps in the project and maps in a plugin.
-        maps_paths = [(None, content_maps_path)]
-
-        maps_plugin_filters = self.MAPS_PLUGIN_FILTERS.get_value()
-        if maps_plugin_filters:
-            plugin_names = self.get_project_plugin_names()
+        # search_paths stores a list of tuples of the form
+        # (unreal_plugin, directory_path). This allows us to differentiate
+        # between maps in the project (unreal_plugin is None in that case) and
+        # maps in a plugin.
+        search_paths = [(None, project_maps_path)]
 
-            for plugin_name in plugin_names:
-                if self.plugin_name_matches(plugin_name, maps_plugin_filters):
-                    maps_paths.append(
-                        (plugin_name, self.get_project_plugin_content_dir(plugin_name)))
+        for unreal_content_plugin in self.get_unreal_content_plugins():
+            search_paths.append(
+                (unreal_content_plugin,
+                 unreal_content_plugin.plugin_content_path))
 
         maps = []
-        for (plugin_name, maps_path) in maps_paths:
+        for (unreal_content_plugin, maps_path) in search_paths:
             for dirpath, b, file_names in os.walk(maps_path):
                 for file_name in file_names:
-                    if not fnmatch.fnmatch(file_name, self.MAPS_FILTER.get_value()):
+                    if not fnmatch.fnmatch(
+                            file_name, self.MAPS_FILTER.get_value()):
                         continue
 
                     map_name, _ = os.path.splitext(file_name)
                     file_path_to_map = os.path.join(dirpath, map_name)
 
                     content_path_to_map = self.resolve_content_path(
-                        file_path_to_map, plugin_name=plugin_name)
+                        file_path_to_map,
+                        unreal_content_plugin=unreal_content_plugin)
 
                     if content_path_to_map not in maps:
                         maps.append(content_path_to_map)
diff --git a/Engine/Plugins/VirtualProduction/Switchboard/Source/Switchboard/switchboard/devices/ndisplay/plugin_ndisplay.py b/Engine/Plugins/VirtualProduction/Switchboard/Source/Switchboard/switchboard/devices/ndisplay/plugin_ndisplay.py
--- a/Engine/Plugins/VirtualProduction/Switchboard/Source/Switchboard/switchboard/devices/ndisplay/plugin_ndisplay.py
+++ b/Engine/Plugins/VirtualProduction/Switchboard/Source/Switchboard/switchboard/devices/ndisplay/plugin_ndisplay.py
@@ -159,33 +159,49 @@ class AddnDisplayDialog(AddDeviceDialog):
     def on_clicked_btnFindConfigs(self):
         ''' Finds and populates config combobox '''
 
-        # We will look for config files in the project's Content folder
-        configs_path = os.path.normpath(
-            os.path.join(
-                os.path.dirname(
-                    CONFIG.UPROJECT_PATH.get_value().replace('"', '')),
-                'Content'))
+        project_configs_path = os.path.normpath(
+            CONFIG.get_project_content_dir())
+
+        # search_paths stores a list of tuples of the form
+        # (unreal_plugin, directory_path). This allows us to differentiate
+        # between nDisplay configs in the project (unreal_plugin is None in
+        # that case) and configs in a plugin.
+        search_paths = [(None, project_configs_path)]
+
+        for unreal_content_plugin in CONFIG.get_unreal_content_plugins():
+            search_paths.append(
+                (unreal_content_plugin,
+                 unreal_content_plugin.plugin_content_path))
+
         config_names = []
         config_paths = []
+        config_plugins = []
 
         assets = []
 
-        for dirpath, _, files in os.walk(configs_path):
-            for name in files:
-                if not name.lower().endswith(('.uasset', '.ndisplay')):
-                    continue
-
-                if name not in config_names:
-                    config_path = os.path.join(dirpath, name)
-                    ext = os.path.splitext(name)[1]
-
-                    # Since .uasset is generic a asset container, only add
-                    # assets of the right config class.
-                    if ext.lower() == '.uasset':
-                        assets.append({'name': name, 'path': config_path})
-                    else:
-                        config_names.append(name)
-                        config_paths.append(config_path)
+        for (unreal_content_plugin, configs_path) in search_paths:
+            for dirpath, _, file_names in os.walk(configs_path):
+                for file_name in file_names:
+                    if not file_name.lower().endswith(
+                            ('.uasset', '.ndisplay')):
+                        continue
+
+                    if file_name not in config_names:
+                        config_path = os.path.join(dirpath, file_name)
+                        ext = os.path.splitext(file_name)[1]
+
+                        # Since .uasset is a generic asset container, only add
+                        # assets of the right config class.
+                        if ext.lower() == '.uasset':
+                            assets.append({
+                                'name': file_name,
+                                'path': config_path,
+                                'plugin': unreal_content_plugin
+                            })
+                        else:
+                            config_names.append(file_name)
+                            config_paths.append(config_path)
+                            config_plugins.append(unreal_content_plugin)
 
         # process the assets in a multi-threaded fashion
 
@@ -234,6 +250,7 @@ class AddnDisplayDialog(AddDeviceDialog):
                     asset = future.result()
                     config_names.append(asset['name'])
                     config_paths.append(asset['path'])
+                    config_plugins.append(asset['plugin'])
                 except Exception:
                     pass
 
@@ -247,7 +264,14 @@ class AddnDisplayDialog(AddDeviceDialog):
         itemDatas = []
 
         for idx, config_name in enumerate(config_names):
-            itemData = {'name': config_name, 'path': config_paths[idx]}
+            config_plugin = config_plugins[idx]
+            uplugin_file_path = (
+                str(config_plugin.uplugin_file_path) if config_plugin
+                else None)
+            itemData = {
+                'name': config_name,
+                'path': config_paths[idx],
+                'uplugin_file_path': uplugin_file_path}
             itemDatas.append(itemData)
 
         # sort by name
diff --git a/Engine/Plugins/VirtualProduction/Switchboard/Source/Switchboard/switchboard/ue_plugin_utils.py b/Engine/Plugins/VirtualProduction/Switchboard/Source/Switchboard/switchboard/ue_plugin_utils.py
new file mode 100644
--- /dev/null
+++ b/Engine/Plugins/VirtualProduction/Switchboard/Source/Switchboard/switchboard/ue_plugin_utils.py
@@ -0,0 +1,162 @@
+# Copyright Epic Games, Inc. All Rights Reserved.
+
+import glob
+import pathlib
+import typing
+
+from switchboard.switchboard_logging import LOGGER
+
+
+class UnrealPlugin(object):
+    '''
+    A simple object that encapsulates properties of an Unreal Engine plugin
+    based on the file system path and name of the plugin's .uplugin file.
+    '''
+
+    @classmethod
+    def from_plugin_path(cls, plugin_path: typing.Union[pathlib.Path, str]):
+        '''
+        Get an UnrealPlugin object based on the given file system path.
+
+        The provided path can be either a file path to an Unreal Engine
+        ".uplugin" file or a path to the plugin directory that contains
+        a ".uplugin" file. In the latter case, the directory will be
+        searched for the ".uplugin" file.
+        '''
+        if not isinstance(plugin_path, pathlib.Path):
+            plugin_path = pathlib.Path(plugin_path)
+
+        if not plugin_path.exists():
+            LOGGER.warning(
+                'Cannot find uplugin file or plugin directory at '
+                f'path: {plugin_path}')
+            return None
+
+        if (plugin_path.is_file() and
+                plugin_path.suffix.casefold() == '.uplugin'):
+            return cls(plugin_path)
+
+        if not plugin_path.is_dir():
+            LOGGER.warning(f'Plugin path is not a directory: {plugin_path}')
+            return None
+
+        uplugin_file_paths = list(plugin_path.glob('*.uplugin'))
+        if not uplugin_file_paths:
+            LOGGER.warning(
+                f'No .uplugin files found in plugin directory: {plugin_path}')
+            return None
+
+        if len(uplugin_file_paths) > 1:
+            msg = (
+                'Multiple .uplugin files found in plugin directory: '
+                f'{plugin_path}\n    ')
+            msg += '\n    '.join([str(p) for p in uplugin_file_paths])
+            LOGGER.warning(msg)
+            return None
+
+        uplugin_file_path = pathlib.Path(uplugin_file_paths[0])
+
+        return cls(uplugin_file_path)
+
+    @classmethod
+    def from_path_filters(
+            cls,
+            ue_project_path: typing.Union[pathlib.Path, str],
+            filter_patterns: typing.List[str]):
+        '''
+        Get a list of Unreal Engine plugins matching the given path-based
+        filter patterns.
+
+        The provided path patterns can be absolute or relative. For relative
+        path patterns, patterns resembling a directory name (i.e. no path
+        separators) will be assumed to refer to plugin directories inside the
+        Unreal Engine project (in its "Plugins" subdirectory). Otherwise,
+        relative path patterns are assumed to be relative project directory
+        (the directory containing the .uproject file).
+        '''
+        if not isinstance(ue_project_path, pathlib.Path):
+            ue_project_path = pathlib.Path(ue_project_path)
+
+        plugin_paths = set()
+
+        for filter_pattern in filter_patterns:
+            path_pattern = pathlib.Path(filter_pattern)
+
+            if path_pattern.is_absolute():
+                plugin_filter = path_pattern
+            else:
+                path_pattern_parts = path_pattern.parts
+                if len(path_pattern_parts) < 1:
+                    continue
+
+                if len(path_pattern_parts) == 1:
+                    # When the path pattern is relative and it looks like a
+                    # directory name (i.e. no path separators), assume we're
+                    # matching against plugin directories inside the project.
+                    plugin_filter = ue_project_path / 'Plugins' / path_pattern
+                else:
+                    # Otherwise, assume that the path pattern is relative to
+                    # the project directory.
+                    plugin_filter = ue_project_path / path_pattern
+
+            for plugin_path in glob.glob(str(plugin_filter)):
+                plugin_path = pathlib.Path(plugin_path).resolve()
+                plugin_paths.add(plugin_path)
+
+        unreal_plugins = []
+        for plugin_path in sorted(list(plugin_paths)):
+            unreal_plugin = cls.from_plugin_path(plugin_path)
+            if unreal_plugin:
+                unreal_plugins.append(unreal_plugin)
+
+        return unreal_plugins
+
+    def __init__(self, uplugin_file_path: typing.Union[pathlib.Path, str]):
+        if not isinstance(uplugin_file_path, pathlib.Path):
+            uplugin_file_path = pathlib.Path(uplugin_file_path)
+
+        self._uplugin_file_path = uplugin_file_path
+
+    def __repr__(self) -> str:
+        return f'{self.__class__.__name__}("{self._uplugin_file_path}")'
+
+    @property
+    def uplugin_file_path(self) -> pathlib.Path:
+        '''
+        The file system path to the plugin's .uplugin file.
+        '''
+        return self._uplugin_file_path
+
+    @property
+    def name(self) -> str:
+        '''
+        The name of the plugin.
+        '''
+        return self._uplugin_file_path.stem
+
+    @property
+    def plugin_path(self) -> pathlib.Path:
+        '''
+        The file system path to the root directory of the plugin.
+        '''
+        return self._uplugin_file_path.parents[0]
+
+    @property
+    def plugin_content_path(self) -> pathlib.Path:
+        '''
+        The file system path to the plugin's "Content" directory.
+
+        This is the directory that gets mounted in Unreal Engine.
+        '''
+        return self.plugin_path / 'Content'
+
+    @property
+    def mounted_path(self) -> pathlib.PurePosixPath:
+        '''
+        The root content path of the plugin when its "Content" directory is
+        mounted in Unreal Engine.
+        '''
+        # Note that UE uses the file name of the uplugin file to produce the
+        # mounted path in engine, and not the name of the plugin directory or
+        # any data inside the uplugin file.
+        return pathlib.PurePosixPath(f'/{self.name}')
