Switchboard: Add Vulkan option to the nDisplay RHI setting dropdown.

diff --git a/Engine/Plugins/VirtualProduction/Switchboard/Source/Switchboard/switchboard/devices/ndisplay/plugin_ndisplay.py b/Engine/Plugins/VirtualProduction/Switchboard/Source/Switchboard/switchboard/devices/ndisplay/plugin_ndisplay.py
--- a/Engine/Plugins/VirtualProduction/Switchboard/Source/Switchboard/switchboard/devices/ndisplay/plugin_ndisplay.py
+++ b/Engine/Plugins/VirtualProduction/Switchboard/Source/Switchboard/switchboard/devices/ndisplay/plugin_ndisplay.py
@@ -406,7 +406,7 @@ class DevicenDisplay(DeviceUnreal):
             attr_name="render_api",
             nice_name="Render API",
             value="dx12",
-            possible_values=["dx11", "dx12"],
+            possible_values=["dx11", "dx12", "vulkan"],
         ),
         'render_mode': OptionSetting(
             attr_name="render_mode",
@@ -848,7 +848,7 @@ class DevicenDisplay(DeviceUnreal):
             f'{session_id}',              # Session ID.
             f'{max_gpu_count}',           # Max GPU count (mGPU)
             f'-dc_cfg="{cfg_file}"',      # nDisplay config file
-            f'{render_api}',              # dx11/12
+            f'{render_api}',              # RHI (dx11/dx12/vulkan)
             f'{render_mode}',             # mono/...
             f'{use_all_cores}',           # -useallavailablecores
             f'{no_texture_streaming}',    # -notexturestreaming
