Switchboard: Fix race condition which could lead to multiple instances of UnrealMultUserSlateServer opening when launching several devices.

Also tidies up some exception handling in add_config_dialog.py, as encountered by a customer in UDN 00477437.

diff --git a/Engine/Plugins/VirtualProduction/Switchboard/Source/Switchboard/switchboard/add_config_dialog.py b/Engine/Plugins/VirtualProduction/Switchboard/Source/Switchboard/switchboard/add_config_dialog.py
--- a/Engine/Plugins/VirtualProduction/Switchboard/Source/Switchboard/switchboard/add_config_dialog.py
+++ b/Engine/Plugins/VirtualProduction/Switchboard/Source/Switchboard/switchboard/add_config_dialog.py
@@ -75,7 +75,7 @@ class AddConfigDialog(QtWidgets.QDialog):
 
         self.setWindowTitle("Add new Switchboard Configuration")
 
-        self.form_layout =  QtWidgets.QFormLayout()
+        self.form_layout = QtWidgets.QFormLayout()
 
         # Path to this configuration file
 
@@ -455,7 +455,16 @@ class AddConfigDialog(QtWidgets.QDialog):
                 # Relying on wmic for this. Another option is to use psutil which is cross-platform, and slower.
                 cmd = f'wmic process where caption="{UEname}" get commandline'
 
-                for line in subprocess.check_output(cmd, startupinfo=sb_utils.get_hidden_sp_startupinfo()).decode().splitlines():
+                try:
+                    commandlines = subprocess.check_output(
+                        cmd, startupinfo=sb_utils.get_hidden_sp_startupinfo()
+                        ).decode().splitlines()
+                except Exception as exc:
+                    LOGGER.error('Exception polling Unreal processes',
+                                 exc_info=exc)
+                    continue
+
+                for line in commandlines:
                     if UEname.lower() not in line.lower():
                         continue
 
diff --git a/Engine/Plugins/VirtualProduction/Switchboard/Source/Switchboard/switchboard/switchboard_application.py b/Engine/Plugins/VirtualProduction/Switchboard/Source/Switchboard/switchboard/switchboard_application.py
--- a/Engine/Plugins/VirtualProduction/Switchboard/Source/Switchboard/switchboard/switchboard_application.py
+++ b/Engine/Plugins/VirtualProduction/Switchboard/Source/Switchboard/switchboard/switchboard_application.py
@@ -87,7 +87,7 @@ class OscServer:
 class MultiUserApplication:
     def __init__(self):
         self.lock = threading.Lock()
-        self.process: Optional[subprocess.Popen] = None
+        self._process: Optional[subprocess.Popen] = None
 
         # Application Options
         self.concert_ignore_cl = False
@@ -111,6 +111,13 @@ class MultiUserApplication:
             return f'-UDPMESSAGING_TRANSPORT_UNICAST="{endpoint}"'
         return ''
 
+    @property
+    def process(self):
+        if self._process and (self._process.poll() is None):
+            return self._process
+        else:
+            return self.poll_process()
+
     def poll_process(self):
         # Aside from this task_name, PollProcess is stateless,
         # and we'd like to pick up on name changes.
@@ -132,7 +139,7 @@ class MultiUserApplication:
             cmdline = ''
             if sys.platform.startswith('win'):
                 if CONFIG.MUSERVER_SLATE_MODE.get_value():
-                    cmdline = f'start /b "Multi User Server" "{self.exe_path()}"'
+                    cmdline = f'{self.exe_path()}'
                 else:
                     cmdline = f'start "Multi User Server" "{self.exe_path()}"'
             else:
@@ -159,32 +166,36 @@ class MultiUserApplication:
                 cmdline += f' {" ".join(args)}'
 
             LOGGER.debug(cmdline)
-            self.process = subprocess.Popen(
+            self._process = subprocess.Popen(
                 cmdline, shell=True,
                 startupinfo=sb_utils.get_hidden_sp_startupinfo())
 
             return True
 
     def terminate(self, bypolling=False):
-        if not bypolling and self.process:
-            self.process.terminate()
+        if not bypolling and self._process:
+            self._process.terminate()
         else:
             self.poll_process().kill()
 
+    FIND_IP_RE = re.compile(
+        r'-UDPMESSAGING_TRANSPORT_UNICAST="(\d+.\d+.\d+.\d+:\d+)"',
+        re.IGNORECASE)
+    FIND_NAME_RE = re.compile(
+        r'-CONCERTSERVER="([A-Za-z0-9_]+)"', re.IGNORECASE)
+
     def extract_process_info(self):
-        pid = self.poll_process().get_pid()
-        command_line = self.poll_process().get_command_line()
+        pid = self.process.pid
+        command_line = str(self.process.args)
+
         if command_line == '' or pid == self._running_pid:
             return
 
-        find_ip_re =  re.compile( r'-UDPMESSAGING_TRANSPORT_UNICAST="(\d+.\d+.\d+.\d+:\d+)"')
-        find_name_re = re.compile(r'-CONCERTSERVER="([A-Za-z0-9_]+)"')
-
         self.clear_process_info()
         self._running_pid = pid
         try:
-            ip_match = find_ip_re.search(command_line)
-            name_match = find_name_re.search(command_line)
+            ip_match = self.FIND_IP_RE.search(command_line)
+            name_match = self.FIND_NAME_RE.search(command_line)
             if ip_match:
                 self._running_endpoint = ip_match.group(1)
             if name_match:
@@ -192,7 +203,6 @@ class MultiUserApplication:
         except:
             pass
 
-
     def validate_process(self):
         if self._running_pid is None:
             return False
@@ -227,7 +237,7 @@ class MultiUserApplication:
         self._running_pid = None
 
     def is_running(self):
-        if self.poll_process().poll() is None:
+        if self.process.poll() is None:
             self.extract_process_info()
             return True
 
diff --git a/Engine/Plugins/VirtualProduction/Switchboard/Source/Switchboard/switchboard/switchboard_utils.py b/Engine/Plugins/VirtualProduction/Switchboard/Source/Switchboard/switchboard/switchboard_utils.py
--- a/Engine/Plugins/VirtualProduction/Switchboard/Source/Switchboard/switchboard/switchboard_utils.py
+++ b/Engine/Plugins/VirtualProduction/Switchboard/Source/Switchboard/switchboard/switchboard_utils.py
@@ -1,13 +1,12 @@
 # Copyright Epic Games, Inc. All Rights Reserved.
+import csv
 import datetime
 from enum import Enum
+import io
 import os
-import pathlib
 import subprocess
-import sys
 import threading
-import io
-import csv
+from typing import Optional
 
 from .switchboard_logging import LOGGER
 
@@ -112,7 +111,8 @@ class PollProcess(object):
     def __init__(self, task_name: str):
         self.task_name = task_name
 
-    def get_command_line(self):
+    @property
+    def args(self) -> str:
         get_commandline_cmd = f"wmic process where caption=\"{self.task_name}\" get commandline"
         try:
             output = subprocess.check_output(get_commandline_cmd, startupinfo=get_hidden_sp_startupinfo()).decode()
@@ -120,15 +120,16 @@ class PollProcess(object):
         except:
             return ""
 
-    def get_pid(self):
+    @property
+    def pid(self) -> Optional[int]:
         pid_cmd = f"tasklist /FI \"IMAGENAME eq {self.task_name}\" /FO csv"
         try:
             output = subprocess.check_output(pid_cmd, startupinfo=get_hidden_sp_startupinfo()).decode()
             dictobj = next(csv.DictReader(io.StringIO(output)))
             if dictobj:
-                return dictobj['PID']
+                return int(dictobj['PID'])
         except:
-            return "Unknown"
+            return None
 
     def poll(self):
         # 'list' output format because default 'table' truncates imagename to 25 characters
