Switchboard: Update uassetparser.py to account for EUnrealEngineObjectUE5Version::ADD_SOFTOBJECTPATH_LIST (introduced in CL 22225193).

#rnx

diff --git a/Engine/Plugins/VirtualProduction/Switchboard/Source/Switchboard/switchboard/devices/unreal/uassetparser.py b/Engine/Plugins/VirtualProduction/Switchboard/Source/Switchboard/switchboard/devices/unreal/uassetparser.py
--- a/Engine/Plugins/VirtualProduction/Switchboard/Source/Switchboard/switchboard/devices/unreal/uassetparser.py
+++ b/Engine/Plugins/VirtualProduction/Switchboard/Source/Switchboard/switchboard/devices/unreal/uassetparser.py
@@ -297,8 +297,13 @@ class UassetParser(object):
 
         s.CustomVersions = self.readTArray(lambda : self.f.read(5*4))        
         s.TotalHeaderSize = self.readInt32()
-        s.FolderName = self.readFString()
+        s.PackageName = self.readFString()
         s.PackageFlags, s.NameCount, s.NameOffset = struct.unpack('<Iii', self.f.read(4*3))
+
+        if s.FileVersionUE5 >= 1008:  # EUnrealEngineObjectUE5Version::ADD_SOFTOBJECTPATH_LIST
+            s.SoftObjectPathsCount = self.readInt32()
+            s.SoftObjectPathsOffset = self.readInt32()
+
         s.LocalizationId = self.readFString()
         
         def remainingBytesInHeader():
