Switchboard: Sets more related cvars and options when unattended and no screen messages.

-unattended is not enough to disable certain popups that interfere with nDisplay, so added Slate.bAllowNotifications=0 to complement.

In addition, Disable All Screen Messages now also generates -NoScreenMessages in command line.

#skipci

diff --git a/Engine/Plugins/VirtualProduction/Switchboard/Source/Switchboard/switchboard/devices/ndisplay/plugin_ndisplay.py b/Engine/Plugins/VirtualProduction/Switchboard/Source/Switchboard/switchboard/devices/ndisplay/plugin_ndisplay.py
--- a/Engine/Plugins/VirtualProduction/Switchboard/Source/Switchboard/switchboard/devices/ndisplay/plugin_ndisplay.py
+++ b/Engine/Plugins/VirtualProduction/Switchboard/Source/Switchboard/switchboard/devices/ndisplay/plugin_ndisplay.py
@@ -825,6 +825,11 @@ class DevicenDisplay(DeviceUnreal):
             if DevicenDisplay.csettings['disable_ensures'].get_value()
             else '')
 
+        no_screen_messages = (
+            '-NoScreenMessages'
+            if DevicenDisplay.csettings['disable_all_screen_messages'].get_value()
+            else '')
+
         # fill in fixed arguments
         args = [
             f'"{uproject}"',
@@ -858,6 +863,7 @@ class DevicenDisplay(DeviceUnreal):
             f'{ini_engine}',              # Engine ini injections
             f'{ini_game}',                # Game ini injections
             f'{unattended}',              # -unattended
+            f'{no_screen_messages}',      # -NoScreenMessages
             f'{disable_ensures}',         # -handleensurepercent=0
             f'{udpm_transport_multi}',    # -UDPMESSAGING_TRANSPORT_MULTICAST=
             f'{udpm_transport_unicast}',  # -UDPMESSAGING_TRANSPORT_UNICAST=
@@ -907,6 +913,10 @@ class DevicenDisplay(DeviceUnreal):
         # Device profile CVars.
         dp_cvars = []
 
+        # If choosing unattended, complement with disabling toasts
+        if DevicenDisplay.csettings['ndisplay_unattended'].get_value():
+            dp_cvars.append('Slate.bAllowNotifications=0')
+
         # Insights traces parameters
         if CONFIG.INSIGHTS_TRACE_ENABLE.get_value():
 
