Switchboard: Add option to disable slate throttling (defaults to no throttling).

This is typically desired on Virtual Production stages since the Editor is used by the operators and prefer to see real time updates on the viewports as they change object properties, such as color grading on led walls.

#skipci

diff --git a/Engine/Plugins/VirtualProduction/Switchboard/Source/Switchboard/switchboard/devices/unreal/plugin_unreal.py b/Engine/Plugins/VirtualProduction/Switchboard/Source/Switchboard/switchboard/devices/unreal/plugin_unreal.py
--- a/Engine/Plugins/VirtualProduction/Switchboard/Source/Switchboard/switchboard/devices/unreal/plugin_unreal.py
+++ b/Engine/Plugins/VirtualProduction/Switchboard/Source/Switchboard/switchboard/devices/unreal/plugin_unreal.py
@@ -422,6 +422,15 @@ class DeviceUnreal(Device):
                 "Expects an absolute path to the device directory containing the ugs.dll library file. \n"
                 "If left blank, switchboard will attempt to find UGS by searching the system PATH and default install locations. \n"
                 "On Windows, the default install location is '${LOCALAPPDATA}/UnrealGameSync/Latest/'.")
+        ),
+        'slate_allow_throttling': BoolSetting(
+            attr_name='slate_allow_throttling',
+            nice_name='Allow Slate Throttling',
+            value=False,
+            tool_tip=(
+                'Sets the Slate.bAllowThrottling cvar. When unchecked, the Editor viewports do not freeze/throttle \n'
+                'during certain operations. Not thottling is typically desired when using the Editor in \n'
+                'a virtual production stage.\n')
         )
         
     }
@@ -1487,6 +1496,10 @@ class DeviceUnreal(Device):
         except ValueError:
             LOGGER.warning(f"Invalid Number of GPUs '{max_gpu_count}'")
 
+        # Slate.bAllowThrottling. Makes ICVFX panel and Vcam more responsive to Editor interactive changes.
+        slate_allow_throttling = DeviceUnreal.csettings["slate_allow_throttling"].get_value()
+        dp_cvars.append(f'Slate.bAllowThrottling={int(slate_allow_throttling)}')
+
         # Add user set dp cvars, overriding any of the forced ones.
         user_dp_cvars = DeviceUnreal.csettings["dp_cvars"].get_value(self.name)
         user_dp_cvars = [cvar.strip() for cvar in user_dp_cvars if len(cvar.strip().split('=')) == 2]
