Switchboard: Added option to not retrieve logs.

Useful for faster iteration when log and trace transfers are not desired.

diff --git a/Engine/Plugins/VirtualProduction/Switchboard/Source/Switchboard/switchboard/devices/unreal/plugin_unreal.py b/Engine/Plugins/VirtualProduction/Switchboard/Source/Switchboard/switchboard/devices/unreal/plugin_unreal.py
--- a/Engine/Plugins/VirtualProduction/Switchboard/Source/Switchboard/switchboard/devices/unreal/plugin_unreal.py
+++ b/Engine/Plugins/VirtualProduction/Switchboard/Source/Switchboard/switchboard/devices/unreal/plugin_unreal.py
@@ -431,8 +431,14 @@ class DeviceUnreal(Device):
                 'Sets the Slate.bAllowThrottling cvar. When unchecked, the Editor viewports do not freeze/throttle \n'
                 'during certain operations. Not thottling is typically desired when using the Editor in \n'
                 'a virtual production stage.\n')
+        ),
+        'retrieve_logs': BoolSetting(
+            attr_name='retrieve_logs',
+            nice_name='Retrieve Logs',
+            value=True,
+            tool_tip=(
+                'When checked, retrieves the logs and traces after Unreal Engine terminates. \n')
         )
-        
     }
 
     unreal_started_signal = QtCore.Signal()
@@ -1726,9 +1732,13 @@ class DeviceUnreal(Device):
                     f'"{program_name}" is still in the list, which is unusual')
 
         if program_name == 'unreal' and not len(remaining_homonyms):
-            log_success = self.start_retrieve_log(unreal_exit_code=returncode)
-            utrace_success = self.start_retrieve_utrace(unreal_exit_code=returncode)
-            if not log_success and not utrace_success:
+            
+            if DeviceUnreal.csettings["retrieve_logs"].get_value():
+                log_success = self.start_retrieve_log(unreal_exit_code=returncode)
+                utrace_success = self.start_retrieve_utrace(unreal_exit_code=returncode)
+                if not log_success and not utrace_success:
+                    self.status = DeviceStatus.CLOSED
+            else:
                 self.status = DeviceStatus.CLOSED
 
         elif program_name == 'retrieve' and not self.transfer_in_progress:
