Switchboard: Added LiveLink dropdown for Unreal and nDisplay devices.

Also fixed bug in the LiveLinkPreset command line parser, where it would think that a future comma later in the command line was the end of the live link name.

diff --git a/Engine/Plugins/Animation/LiveLink/Source/LiveLink/Private/LiveLinkModule.cpp b/Engine/Plugins/Animation/LiveLink/Source/LiveLink/Private/LiveLinkModule.cpp
--- a/Engine/Plugins/Animation/LiveLink/Source/LiveLink/Private/LiveLinkModule.cpp
+++ b/Engine/Plugins/Animation/LiveLink/Source/LiveLink/Private/LiveLinkModule.cpp
@@ -75,21 +75,29 @@ void FLiveLinkModule::OnEngineLoopInitComplete()
 {
 	ULiveLinkPreset* StartupPreset = nullptr;
 	const FString& CommandLine = FCommandLine::Get();
-	const TCHAR* PresetStr = TEXT("LiveLink.Preset.Apply Preset=");
-	const int32 FoundElement = CommandLine.Find(PresetStr);
-	if (FoundElement != INDEX_NONE)
+	const TCHAR* PresetStr = TEXT("LiveLink.Preset.Apply Preset="); // expected inside an -ExecCmds="" argument. So our command should end either on ',' or '"'.
+	const int32 CommandStartPos = CommandLine.Find(PresetStr);
+
+	if (CommandStartPos != INDEX_NONE)
 	{
-		int32 LiveLinkArgumentEnd = CommandLine.Find(",", ESearchCase::IgnoreCase, ESearchDir::FromStart, FoundElement);
-		if (LiveLinkArgumentEnd == INDEX_NONE)
+		int32 PresetEndPos = CommandLine.Find(",", ESearchCase::IgnoreCase, ESearchDir::FromStart, CommandStartPos);
+		const int32 NextDoubleQuotesPos = CommandLine.Find("\"", ESearchCase::IgnoreCase, ESearchDir::FromStart, CommandStartPos);
+
+		if ((PresetEndPos != INDEX_NONE) && (NextDoubleQuotesPos != INDEX_NONE))
 		{
-			LiveLinkArgumentEnd = CommandLine.Find("\"", ESearchCase::IgnoreCase, ESearchDir::FromStart, FoundElement);
+			PresetEndPos = FMath::Min(PresetEndPos, NextDoubleQuotesPos);
 		}
-		if (LiveLinkArgumentEnd != INDEX_NONE)
+		else if (NextDoubleQuotesPos != INDEX_NONE)
+		{
+			PresetEndPos = NextDoubleQuotesPos;
+		}
+
+		if (PresetEndPos != INDEX_NONE)
 		{
-			const int32 StartIndex = FoundElement + FCString::Strlen(PresetStr);
-			if (CommandLine.IsValidIndex(StartIndex) && CommandLine.IsValidIndex(LiveLinkArgumentEnd))
+			const int32 PresetStartPos = CommandStartPos + FCString::Strlen(PresetStr);
+			if (CommandLine.IsValidIndex(PresetStartPos) && CommandLine.IsValidIndex(PresetEndPos))
 			{
-				const FString LiveLinkPresetName = CommandLine.Mid(StartIndex, LiveLinkArgumentEnd - StartIndex);
+				const FString LiveLinkPresetName = CommandLine.Mid(PresetStartPos, PresetEndPos - PresetStartPos);
 				StartupPreset = Cast<ULiveLinkPreset>(StaticLoadObject(ULiveLinkPreset::StaticClass(), nullptr, *LiveLinkPresetName));
 			}
 		}
diff --git a/Engine/Plugins/VirtualProduction/Switchboard/Source/Switchboard/switchboard/devices/ndisplay/plugin_ndisplay.py b/Engine/Plugins/VirtualProduction/Switchboard/Source/Switchboard/switchboard/devices/ndisplay/plugin_ndisplay.py
--- a/Engine/Plugins/VirtualProduction/Switchboard/Source/Switchboard/switchboard/devices/ndisplay/plugin_ndisplay.py
+++ b/Engine/Plugins/VirtualProduction/Switchboard/Source/Switchboard/switchboard/devices/ndisplay/plugin_ndisplay.py
@@ -21,7 +21,7 @@ from switchboard.config import CONFIG, BoolSetting, FilePathSetting, \
     StringListSetting, migrate_comma_separated_string_to_list
 from switchboard.devices.device_widget_base import AddDeviceDialog
 from switchboard.devices.unreal.plugin_unreal import DeviceUnreal, \
-    DeviceWidgetUnreal
+    DeviceWidgetUnreal, LiveLinkPresetSetting
 from switchboard.devices.unreal.uassetparser import UassetParser
 from switchboard.switchboard_logging import LOGGER
 
@@ -96,17 +96,33 @@ class AddnDisplayDialog(AddDeviceDialog):
         Populate the config combobox with the already found configs, avoiding
         having to re-find every time.
         '''
+
+         # cache the current value of the combo box (currently selected preset for this device)
+        cur_item = self.cbConfigs.currentData()
+
         self.cbConfigs.clear()
 
-        itemDatas = DevicenDisplay.csettings[
-            'populated_config_itemDatas'].get_value()
+        itemDatas = DeviceUnreal.csettings['asset_itemDatas'].get_value()
+
+        def asset_is_nDisplay_config(itemData:dict) -> bool:
+            ''' Convenience function to filter the nDisplay configs '''
+
+            # Only return assets of the correct class
+            return itemData['classname'] in DeviceUnreal.NDISPLAY_CLASS_NAMES
 
         try:
-            for itemData in itemDatas:
+            for itemData in [itemData for itemData in itemDatas if asset_is_nDisplay_config(itemData)]:
                 self.cbConfigs.addItem(itemData['name'], itemData)
         except Exception:
             LOGGER.error('Error recalling config itemDatas')
 
+        # restore selected item
+        for item_idx in range(self.cbConfigs.count()):
+            if cur_item and (cur_item['name'] == self.cbConfigs.itemData(item_idx)['name']):
+                self.cbConfigs.setCurrentIndex(item_idx)
+                break
+
+
     def current_config_path(self):
         ''' Get currently selected config path in the combobox '''
 
@@ -159,127 +175,7 @@ class AddnDisplayDialog(AddDeviceDialog):
     def on_clicked_btnFindConfigs(self):
         ''' Finds and populates config combobox '''
 
-        project_configs_path = os.path.normpath(
-            CONFIG.get_project_content_dir())
-
-        # search_paths stores a list of tuples of the form
-        # (unreal_plugin, directory_path). This allows us to differentiate
-        # between nDisplay configs in the project (unreal_plugin is None in
-        # that case) and configs in a plugin.
-        search_paths = [(None, project_configs_path)]
-
-        for unreal_content_plugin in CONFIG.get_unreal_content_plugins():
-            search_paths.append(
-                (unreal_content_plugin,
-                 unreal_content_plugin.plugin_content_path))
-
-        config_names = []
-        config_paths = []
-        config_plugins = []
-
-        assets = []
-
-        for (unreal_content_plugin, configs_path) in search_paths:
-            for dirpath, _, file_names in os.walk(configs_path):
-                for file_name in file_names:
-                    if not file_name.lower().endswith(
-                            ('.uasset', '.ndisplay')):
-                        continue
-
-                    if file_name not in config_names:
-                        config_path = os.path.join(dirpath, file_name)
-                        ext = os.path.splitext(file_name)[1]
-
-                        # Since .uasset is a generic asset container, only add
-                        # assets of the right config class.
-                        if ext.lower() == '.uasset':
-                            assets.append({
-                                'name': file_name,
-                                'path': config_path,
-                                'plugin': unreal_content_plugin
-                            })
-                        else:
-                            config_names.append(file_name)
-                            config_paths.append(config_path)
-                            config_plugins.append(unreal_content_plugin)
-
-        # process the assets in a multi-threaded fashion
-
-        # show a progress bar if it is taking more a trivial amount of time
-        progressDiag = QtWidgets.QProgressDialog(
-            'Parsing assets...', 'Cancel', 0, 0, parent=self)
-        progressDiag.setWindowTitle('nDisplay Config Finder')
-        progressDiag.setModal(True)
-        progressDiag.setMinimumDuration(1000)  # time before it shows up
-        progressDiag.setRange(0, len(assets))
-        progressDiag.setCancelButton(None)
-        # Looks much better without the window frame.
-        progressDiag.setWindowFlag(QtCore.Qt.FramelessWindowHint)
-
-        ''' Returns the asset if it is an nDisplay config '''
-        def validateConfigAsset(asset):
-
-            with open(asset['path'], 'rb') as file:
-
-                aparser = UassetParser(file, allowUnversioned=True)
-
-                for assetdata in aparser.aregdata:
-                    if assetdata.ObjectClassName in DisplayConfig.VALID_CLASS_NAMES:
-                        return asset
-
-            raise ValueError
-
-        numThreads = 8
-        doneAssetCount = 0
-
-        with concurrent.futures.ThreadPoolExecutor(
-                max_workers=numThreads) as executor:
-            futures = [
-                executor.submit(validateConfigAsset, asset)
-                for asset in assets]
-
-            for future in concurrent.futures.as_completed(futures):
-
-                # Update progress bar.
-                doneAssetCount += 1
-                progressDiag.setValue(doneAssetCount)
-
-                # Get the future result and add to list of config names and
-                # paths.
-                try:
-                    asset = future.result()
-                    config_names.append(asset['name'])
-                    config_paths.append(asset['path'])
-                    config_plugins.append(asset['plugin'])
-                except Exception:
-                    pass
-
-        config_names, _ = sb_dialog.SwitchboardDialog.generate_disambiguated_names(config_paths, AddnDisplayDialog.generate_short_unique_config_name)
-
-        # close progress bar window
-        progressDiag.close()
-
-        # collect the found config files into the itemDatas list
-
-        itemDatas = []
-
-        for idx, config_name in enumerate(config_names):
-            config_plugin = config_plugins[idx]
-            uplugin_file_path = (
-                str(config_plugin.uplugin_file_path) if config_plugin
-                else None)
-            itemData = {
-                'name': config_name,
-                'path': config_paths[idx],
-                'uplugin_file_path': uplugin_file_path}
-            itemDatas.append(itemData)
-
-        # sort by name
-        itemDatas.sort(key=lambda itemData: itemData['name'])
-
-        # update settings that should survive device removal and addition
-        DevicenDisplay.csettings['populated_config_itemDatas'].update_value(
-            itemDatas)
+        DeviceUnreal.analyze_project_assets()
 
         # update the combo box with the items.
         self.recall_config_itemDatas()
@@ -369,11 +265,6 @@ class DeviceWidgetnDisplay(DeviceWidgetUnreal):
 class DisplayConfig(object):
     ''' Encapsulates nDisplay config'''
 
-    VALID_CLASS_NAMES = (
-        'DisplayClusterBlueprint',
-        '/Script/DisplayCluster.DisplayClusterBlueprint'
-    )
-
     def __init__(self):
         self.nodes = []
         self.uasset_path = ''
@@ -553,6 +444,13 @@ class DevicenDisplay(DeviceUnreal):
             value=True,
             tool_tip="When checked, adds DisableAllScreenMessages to ExecCmds"
         ),
+        'livelink_preset': LiveLinkPresetSetting(
+            attr_name='livelink_preset',
+            nice_name='LiveLink Preset',
+            value='',
+            tool_tip=(
+                'Adds the selected LiveLink preset to the command line \n')
+        ),
     }
 
     ndisplay_monitor_ui = None
@@ -681,6 +579,7 @@ class DevicenDisplay(DeviceUnreal):
             DevicenDisplay.csettings['priority_modifier'],
             DevicenDisplay.csettings['udpmessaging_unicast_endpoint'],
             DevicenDisplay.csettings['udpmessaging_extra_static_endpoints'],
+            DevicenDisplay.csettings['livelink_preset'],
             CONFIG.ENGINE_DIR,
             CONFIG.SOURCE_CONTROL_WORKSPACE,
             CONFIG.UPROJECT_PATH,
@@ -875,7 +774,23 @@ class DevicenDisplay(DeviceUnreal):
 
         if DevicenDisplay.csettings['disable_all_screen_messages'].get_value() and 'DisableAllScreenMessages' not in exec_cmds:
             exec_cmds.append('DisableAllScreenMessages')
+        
+        # LiveLink presets can be applied as ExecCmds
+        # e.g. of command:
+        #   "LiveLink.Preset.Apply Preset=/Game/Folder/MyLiveLinkPreset.MyLiveLinkPreset"
+
+        livelink_preset_gamepath = DeviceUnreal.csettings["livelink_preset"].get_value(self.name)
+
+        if livelink_preset_gamepath:
+
+            ext = '.uasset'
+
+            if livelink_preset_gamepath.endswith(ext):
+                livelink_preset_gamepath = livelink_preset_gamepath[:len(livelink_preset_gamepath)-len(ext)]
             
+            preset_name = os.path.basename(os.path.normpath(livelink_preset_gamepath))
+            exec_cmds.append(f"LiveLink.Preset.Apply Preset={livelink_preset_gamepath}.{preset_name}")
+
         exec_cmds = [cmd for cmd in exec_cmds if len(cmd.strip())]
 
         if len(exec_cmds):
@@ -1063,7 +978,7 @@ class DevicenDisplay(DeviceUnreal):
             # ConfigExport tag and parse it.
 
             for assetdata in aparser.aregdata:
-                if assetdata.ObjectClassName in DisplayConfig.VALID_CLASS_NAMES:
+                if assetdata.ObjectClassName in DeviceUnreal.NDISPLAY_CLASS_NAMES:
                     return assetdata.tags['ConfigExport']
 
         raise ValueError('Invalid nDisplay config .uasset')
diff --git a/Engine/Plugins/VirtualProduction/Switchboard/Source/Switchboard/switchboard/devices/unreal/plugin_unreal.py b/Engine/Plugins/VirtualProduction/Switchboard/Source/Switchboard/switchboard/devices/unreal/plugin_unreal.py
--- a/Engine/Plugins/VirtualProduction/Switchboard/Source/Switchboard/switchboard/devices/unreal/plugin_unreal.py
+++ b/Engine/Plugins/VirtualProduction/Switchboard/Source/Switchboard/switchboard/devices/unreal/plugin_unreal.py
@@ -3,46 +3,42 @@
 from __future__ import annotations
 
 import base64
+import concurrent.futures
 import json
-from collections import OrderedDict
-from datetime import datetime
-from functools import wraps
-from ipaddress import IPv4Address
 import os
 import pathlib
 import re
 import socket
 import sys
 import threading
-from typing import Callable, List, Optional, Set
 import uuid
+from collections import OrderedDict
+from datetime import datetime
+from functools import wraps
+from ipaddress import IPv4Address
+from typing import Callable, List, Optional, Set
 
-from PySide2 import QtCore
-from PySide2 import QtGui
-from PySide2 import QtWidgets
-
+import switchboard.switchboard_widgets as sb_widgets
+from PySide2 import QtCore, QtGui, QtWidgets
 from switchboard import config_osc as osc
-from switchboard import message_protocol
-from switchboard import switchboard_application
+from switchboard import message_protocol, switchboard_application
+from switchboard import switchboard_dialog as sb_dialog
 from switchboard import switchboard_utils as sb_utils
-from switchboard.config import CONFIG, BoolSetting, DirectoryPathSetting, \
-    FilePathSetting, IntSetting, MultiOptionSetting, OptionSetting, StringSetting, \
-    SETTINGS, DEFAULT_MAP_TEXT, StringListSetting, migrate_comma_separated_string_to_list, \
-    EngineSyncMethod
-from switchboard.devices.device_base import Device, DeviceStatus, \
-    PluginHeaderWidgets
-from switchboard.devices.device_widget_base import DeviceWidget, DeviceAutoJoinMUServerUI
+from switchboard.config import (CONFIG, DEFAULT_MAP_TEXT, SETTINGS, BoolSetting, DirectoryPathSetting, EngineSyncMethod,
+                                FilePathSetting, IntSetting, MultiOptionSetting, OptionSetting, Setting,
+                                StringListSetting, StringSetting, migrate_comma_separated_string_to_list)
+from switchboard.devices.device_base import Device, DeviceStatus, PluginHeaderWidgets
+from switchboard.devices.device_widget_base import DeviceAutoJoinMUServerUI, DeviceWidget
+from switchboard.devices.unreal.uassetparser import UassetParser
 from switchboard.listener_client import ListenerClient
 from switchboard.switchboard_logging import LOGGER
-import switchboard.switchboard_widgets as sb_widgets
-
 from switchboard.tools.insights_launcher import InsightsLauncher
 
-from .listener_watcher import ListenerWatcher
-from .redeploy_dialog import RedeployListenerDialog
-from . import version_helpers
 from ...util import p4_changelist_inspection
 from ...util.p4_changelist_inspection import P4Error
+from . import version_helpers
+from .listener_watcher import ListenerWatcher
+from .redeploy_dialog import RedeployListenerDialog
 
 
 class ProgramStartQueueItem:
@@ -269,9 +265,171 @@ class ProgramStartQueue:
     def clear_running_programs(self):
         self.running_programs.clear()
 
+class LiveLinkPresetSetting(Setting):
+    ''' Container of the LiveLink Preset setting
+    Its main widget is a combo box that makes available the list
+    of presets detected in the project and content plugins.
+
+    Its "refresh" button triggers an asset traversal to refresh
+    the options.
+    '''
+
+    def __init__(
+        self,
+        attr_name,
+        nice_name,
+        value,
+        tool_tip=None,
+        show_ui=True,
+        allow_reset=True,
+        migrate_data=None
+    ):
+        super().__init__(
+            attr_name=attr_name,
+            nice_name=nice_name,
+            value=value,
+            tool_tip=tool_tip,
+            show_ui=show_ui,
+            allow_reset=allow_reset,
+            migrate_data=migrate_data)
+
+    def _create_widgets(self, override_device_name = None):
+
+        # create combo with livelink preset options
+
+        combo = sb_widgets.SearchableComboBox(None)
+        combo.setSizePolicy(QtWidgets.QSizePolicy.Expanding, QtWidgets.QSizePolicy.Preferred)
+        self._update_combo_items(combo, override_device_name)
+
+        combo.currentTextChanged.connect(
+            lambda: self._validate_and_commit_value(combo, override_device_name)
+        )
+
+        # create refresh button
+
+        btn_refresh = sb_widgets.ControlQPushButton()
+        btn_refresh.setIcon(QtGui.QIcon(':/icons/images/icon_refresh.png'))
+        btn_refresh.setProperty("frameless", True)
+        btn_refresh.setToolTip("Refresh list")
+        btn_refresh.clicked.connect(
+            lambda: self._refresh_livelink_preset_list(combo, override_device_name)
+        )
+
+        # create layout with combo and refresh button
+
+        layout = QtWidgets.QHBoxLayout()
+
+        layout.addWidget(combo)
+        layout.addWidget(btn_refresh)
+
+        # set the class' widget as the combo box itself
+        self.set_widget(widget=combo, override_device_name=override_device_name)
+
+        return layout
+
+    def _refresh_livelink_preset_list(self, combo:QtWidgets.QComboBox, override_device_name:Optional[str] = None):
+        ''' Traveses the content to refresh the list of live link presets available '''
+
+        DeviceUnreal.analyze_project_assets()
+        self._update_combo_items(combo, override_device_name)
+
+    def _validate_and_commit_value(self, combo: QtWidgets.QComboBox, override_device_name:Optional[str] = None):
+
+        itemData = combo.currentData()
+
+        value_str = ''
+
+        if itemData:
+            value_str = itemData['gamepath'] # we use path and not name because it is unambiguous
+
+        self._on_widget_value_changed(value_str, override_device_name=override_device_name)
+
+    def _update_combo_items(self, combo:QtWidgets.QComboBox, override_device_name):
+        '''
+        Populate the combobox itself with known list of live link presets available.
+        Makes sure that the currently selected preset is preserved, unless it is not found,
+        in which case it defaults to no preset.
+        '''
+
+        # cache the current value of the combo box (currently selected preset for this device)
+        cur_value = self.get_value(override_device_name)
+
+        # We will re-populate the combo box from scratch
+        combo.clear()
+
+        # add the empty/none choice
+        combo.addItem('', {'name':'', 'gamepath':''})
+
+        itemDatas = DeviceUnreal.csettings['asset_itemDatas'].get_value()
+
+        def asset_is_livelink_preset(itemData) -> bool:
+            ''' Convenience function to filter the live link preset assets '''
+
+            # Only return assets of the correct live link preset class
+            if itemData['classname'] not in DeviceUnreal.LIVELINKPRESET_CLASS_NAMES:
+                return False
+                
+            # Only return live link presets that have a valid gamepath. 
+            # We wouldn't be able to build the command line otherwise
+            if itemData['gamepath'] == '':
+                return False
+
+            return True
+
+        # generate the combo box items
+        try:
+            for itemData in [itemData for itemData in itemDatas if asset_is_livelink_preset(itemData)]:
+                name = itemData['name']
+
+                # trim the expected .uasset extension
+                ext = '.uasset'
+                if name.endswith(ext):
+                    name = name[:len(name)-len(ext)]
+
+                combo.addItem(name, itemData)
+
+        except Exception:
+            LOGGER.error('Error recalling asset itemDatas')
+
+        # set the current index to the live link preset that was already selected
+        for item_idx in range(combo.count()):
+            if cur_value == combo.itemData(item_idx)['gamepath']:
+                combo.setCurrentIndex(item_idx)
+                break
+
+    def _on_setting_changed(self, new_value: str, override_device_name: Optional[str] = None):
+
+        combo = self.get_widget(override_device_name=override_device_name)
+
+        if not combo:
+            return
+
+        try:
+            old_value = combo.currentData()['gamepath']
+        except (KeyError, TypeError):
+            old_value = ''
+
+        new_str_value = new_value
+
+        # if the value changed, find the new index in the combo box based on gamepath
+        if new_str_value != old_value:
+            for item_idx in range(combo.count()):
+                if new_str_value == combo.itemData(item_idx)['gamepath']:
+                    combo.setCurrentIndex(item_idx)
+                    break
 
 class DeviceUnreal(Device):
 
+    NDISPLAY_CLASS_NAMES =  (
+        'DisplayClusterBlueprint',
+        '/Script/DisplayCluster.DisplayClusterBlueprint',
+    )
+
+    LIVELINKPRESET_CLASS_NAMES = (
+        'LiveLinkPreset',
+        '/Script/LiveLink.LiveLinkPreset',
+    )
+    
     csettings = {
         'buffer_size': IntSetting(
             attr_name="buffer_size",
@@ -438,7 +596,21 @@ class DeviceUnreal(Device):
             value=True,
             tool_tip=(
                 'When checked, retrieves the logs and traces after Unreal Engine terminates. \n')
-        )
+        ),
+        'livelink_preset': LiveLinkPresetSetting(
+            attr_name='livelink_preset',
+            nice_name='LiveLink Preset',
+            value='',
+            tool_tip=(
+                'Adds the selected LiveLink preset to the command line \n')
+        ),
+        'asset_itemDatas': Setting(
+            attr_name='asset_itemDatas',
+            nice_name="Asset files",
+            value=[],
+            tool_tip="Remember the last analyzed list of project and plugin assets",
+            show_ui=False,
+        ),
     }
 
     unreal_started_signal = QtCore.Signal()
@@ -776,6 +948,7 @@ class DeviceUnreal(Device):
             DeviceUnreal.csettings['udpmessaging_unicast_endpoint'],
             DeviceUnreal.csettings['udpmessaging_extra_static_endpoints'],
             DeviceUnreal.csettings['unrealgamesync_lib_dir'],
+            DeviceUnreal.csettings['livelink_preset'],
             CONFIG.ENGINE_DIR,
             CONFIG.SOURCE_CONTROL_WORKSPACE,
             CONFIG.UPROJECT_PATH,
@@ -1468,7 +1641,27 @@ class DeviceUnreal(Device):
                 CONFIG.INSIGHTS_TRACE_ARGS.get_value())
 
         exec_cmds = DeviceUnreal.csettings["exec_cmds"].get_value(self.name).copy()
+
+        # LiveLink presets can be applied as ExecCmds
+        # e.g. of command:
+        #   "LiveLink.Preset.Apply Preset=/Game/Folder/MyLiveLinkPreset.MyLiveLinkPreset"
+
+        livelink_preset_gamepath = DeviceUnreal.csettings["livelink_preset"].get_value(self.name)
+
+        if livelink_preset_gamepath:
+
+            ext = '.uasset'
+
+            if livelink_preset_gamepath.endswith(ext):
+                livelink_preset_gamepath = livelink_preset_gamepath[:len(livelink_preset_gamepath)-len(ext)]
+            
+            preset_name = os.path.basename(os.path.normpath(livelink_preset_gamepath))
+            exec_cmds.append(f"LiveLink.Preset.Apply Preset={livelink_preset_gamepath}.{preset_name}")
+
+        # Exec Commands
+
         exec_cmds = [cmd for cmd in exec_cmds if len(cmd.strip())]
+
         if len(exec_cmds):
             exec_cmds_expanded = ','.join(exec_cmds)
             command_line_args += f' -ExecCmds="{exec_cmds_expanded}"'
@@ -2239,6 +2432,151 @@ class DeviceUnreal(Device):
         QtGui.QGuiApplication.clipboard().setText(
             self.last_launch_command.get_value())
 
+    @classmethod
+    def analyze_project_assets(cls):
+        ''' Traverses project and content plugins and caches a list of assets of interest. 
+        That list is currently nDisplay configs and live link presets.
+        '''
+
+        project_configs_path = os.path.normpath(
+            CONFIG.get_project_content_dir())
+
+        # search_paths stores a list of tuples of the form
+        # (unreal_plugin, directory_path). This allows us to differentiate
+        # between project assets (unreal_plugin is None in
+        # that case) and plugin assets.
+        search_paths = [(None, project_configs_path)]
+
+        for unreal_content_plugin in CONFIG.get_unreal_content_plugins():
+            search_paths.append(
+                (unreal_content_plugin,
+                 unreal_content_plugin.plugin_content_path))
+
+        asset_names = []
+        asset_paths = []
+        asset_plugins = []
+        asset_classnames = []
+
+        assets = []
+
+        for (unreal_content_plugin, configs_path) in search_paths:
+            for dirpath, _, file_names in os.walk(configs_path):
+                for file_name in file_names:
+                    if not file_name.lower().endswith(('.uasset', '.ndisplay')):
+                        continue
+
+                    if file_name not in asset_names:
+                        asset_path = os.path.join(dirpath, file_name)
+                        ext = os.path.splitext(file_name)[1]
+
+                        # Since .uasset is a generic asset container, only add
+                        # assets of the right class.
+                        if ext.lower() == '.uasset':
+                            assets.append({
+                                'name': file_name,
+                                'path': asset_path,
+                                'plugin': unreal_content_plugin,
+                            })
+                        else:
+                            asset_names.append(file_name)
+                            asset_paths.append(asset_path)
+                            asset_plugins.append(unreal_content_plugin)
+
+        # process the assets in a multi-threaded fashion
+
+        # show a progress bar if it is taking more a trivial amount of time
+        progressDiag = QtWidgets.QProgressDialog(
+            'Parsing assets...', 'Cancel', 0, 0, parent=None)
+
+        progressDiag.setWindowTitle('Unreal Asset Analyzer')
+        progressDiag.setModal(True)
+        progressDiag.setMinimumDuration(1000)  # time before it shows up
+        progressDiag.setRange(0, len(assets))
+        progressDiag.setCancelButton(None)
+
+        # Looks much better without the window frame.
+        progressDiag.setWindowFlag(QtCore.Qt.FramelessWindowHint)
+
+        INTERESTING_CLASS_NAMES = DeviceUnreal.NDISPLAY_CLASS_NAMES + DeviceUnreal.LIVELINKPRESET_CLASS_NAMES
+
+        def validateInterestingAsset(asset):
+            ''' Returns the asset if it is an interesting asset '''
+
+            with open(asset['path'], 'rb') as file:
+
+                aparser = UassetParser(file, allowUnversioned=True)
+
+                for assetdata in aparser.aregdata:
+                    if assetdata.ObjectClassName in INTERESTING_CLASS_NAMES:
+                        asset['assetdata'] = assetdata
+                        return asset
+
+            raise ValueError
+
+        numThreads = 8
+        doneAssetCount = 0
+
+        with concurrent.futures.ThreadPoolExecutor(
+                max_workers=numThreads) as executor:
+            futures = [
+                executor.submit(validateInterestingAsset, asset)
+                for asset in assets]
+
+            for future in concurrent.futures.as_completed(futures):
+
+                # Update progress bar.
+                doneAssetCount += 1
+                progressDiag.setValue(doneAssetCount)
+
+                # Get the future result and add to list of config names and
+                # paths.
+                try:
+                    asset = future.result()
+                    asset_names.append(asset['name'])
+                    asset_paths.append(asset['path'])
+                    asset_plugins.append(asset['plugin'])
+                    asset_classnames.append(asset['assetdata'].ObjectClassName)
+                except Exception:
+                    pass
+
+        # close progress bar window
+        progressDiag.close()
+
+        def generate_short_unique_config_name(config_path: str, file_name: str) -> str:
+            config_path = CONFIG.shrink_path(config_path)
+            return sb_dialog.SwitchboardDialog.filter_empty_abiguated_path(config_path, file_name)
+
+        asset_names, _ = sb_dialog.SwitchboardDialog.generate_disambiguated_names(asset_paths, generate_short_unique_config_name)
+
+        # collect the found config files into the itemDatas list
+
+        itemDatas = []
+
+        for idx, asset_name in enumerate(asset_names):
+            
+            uplugin_file_path = (
+                str(asset_plugins[idx].uplugin_file_path) if asset_plugins[idx]
+                else None)
+
+            gamepath = CONFIG.resolve_content_path(
+                file_path=asset_paths[idx], 
+                unreal_content_plugin=asset_plugins[idx])
+
+            itemData = {
+                'name': asset_name,
+                'path': asset_paths[idx],
+                'uplugin_file_path': uplugin_file_path,
+                'classname': asset_classnames[idx],
+                'gamepath': gamepath,
+            }
+            itemDatas.append(itemData)
+
+        # sort by name
+        itemDatas.sort(key=lambda itemData: itemData['name'])
+
+        # update settings that should survive device removal and addition
+        DeviceUnreal.csettings['asset_itemDatas'].update_value(itemDatas)
+
 def parse_unreal_tag_file(file_content):
     tags = []
     for line in file_content:
