SwitchBoard: Added RenderOffscreen parameter to command line

diff --git a/Engine/Plugins/VirtualProduction/Switchboard/Source/Switchboard/switchboard/devices/ndisplay/plugin_ndisplay.py b/Engine/Plugins/VirtualProduction/Switchboard/Source/Switchboard/switchboard/devices/ndisplay/plugin_ndisplay.py
--- a/Engine/Plugins/VirtualProduction/Switchboard/Source/Switchboard/switchboard/devices/ndisplay/plugin_ndisplay.py
+++ b/Engine/Plugins/VirtualProduction/Switchboard/Source/Switchboard/switchboard/devices/ndisplay/plugin_ndisplay.py
@@ -485,6 +485,12 @@ class DevicenDisplay(DeviceUnreal):
                 value=kwargs.get("fullscreen", False),
                 show_ui=False
             ),
+            'headless' : BoolSetting(
+                attr_name="headless",
+                nice_name="headless",
+                value=kwargs.get("renderHeadless", False),
+                show_ui=False
+            ),
         }
 
         self.render_mode_cmdline_opts = {
@@ -627,6 +633,7 @@ class DevicenDisplay(DeviceUnreal):
         win_pos = self.settings['window_position'].get_value()
         win_res = self.settings['window_resolution'].get_value()
         fullscreen = self.settings['fullscreen'].get_value()
+        headless = self.settings['headless'].get_value()
 
         # Misc settings
         render_mode = self.render_mode_cmdline_opts[
@@ -812,6 +819,11 @@ class DevicenDisplay(DeviceUnreal):
                 f'ResY={win_res[1]}',
             ])
 
+        if headless:
+            args.extend([
+                '-RenderOffscreen',
+            ])
+
         # MultiUser parameters
         if CONFIG.MUSERVER_AUTO_JOIN.get_value() and self.autojoin_mu_server.get_value():
             args.extend([
@@ -1116,8 +1128,9 @@ class DevicenDisplay(DeviceUnreal):
 
             kwargs["window_position"] = (winx, winy)
             kwargs["window_resolution"] = (resx, resy)
-            # Note the capital 'S'.
+            # Note the capital 'S', 'H'...
             kwargs["fullscreen"] = bool(cnode.get('fullScreen', False))
+            kwargs["headless"] = bool(cnode.get('renderHeadless', False))
 
             primary = True if primaryNode['id'] == name else False
 
@@ -1177,6 +1190,8 @@ class DevicenDisplay(DeviceUnreal):
             menode['kwargs'].get("window_resolution", (100, 100)))
         self.settings['fullscreen'].update_value(
             menode['kwargs'].get("fullscreen", False))
+        self.settings['headless'].update_value(
+            menode['kwargs'].get("headless", False))
 
     @classmethod
     def uasset_path_from_object_path(
