Switchboard: Added UnrealTasks to default channels.

Also added "default" and removed those included by it, simplifying the list.

TraceAuxiliary.cpp:

GDefaultChannels = TEXT("cpu,gpu,frame,log,bookmark");
GMemoryChannels = TEXT("memtag,memalloc,callstack,module");

#skipci

diff --git a/Engine/Plugins/VirtualProduction/Switchboard/Source/Switchboard/switchboard/config.py b/Engine/Plugins/VirtualProduction/Switchboard/Source/Switchboard/switchboard/config.py
--- a/Engine/Plugins/VirtualProduction/Switchboard/Source/Switchboard/switchboard/config.py
+++ b/Engine/Plugins/VirtualProduction/Switchboard/Source/Switchboard/switchboard/config.py
@@ -2173,7 +2173,7 @@ class Config(object):
             "tracing_args": StringSetting(
                 "tracing_args",
                 "Unreal Insights Tracing Args",
-                data.get('tracing_args', 'log,cpu,gpu,frame,bookmark,concert,messaging')
+                data.get('tracing_args', 'default,concert,messaging,tasks')
             ),
             "tracing_stat_events": BoolSetting(
                 "tracing_stat_events",
