Switchboard: Also turning on vsync bookmarks when "default" trace is enabled, since it implies bookmarks as well.

#skipci

diff --git a/Engine/Plugins/VirtualProduction/Switchboard/Source/Switchboard/switchboard/devices/ndisplay/plugin_ndisplay.py b/Engine/Plugins/VirtualProduction/Switchboard/Source/Switchboard/switchboard/devices/ndisplay/plugin_ndisplay.py
--- a/Engine/Plugins/VirtualProduction/Switchboard/Source/Switchboard/switchboard/devices/ndisplay/plugin_ndisplay.py
+++ b/Engine/Plugins/VirtualProduction/Switchboard/Source/Switchboard/switchboard/devices/ndisplay/plugin_ndisplay.py
@@ -862,7 +862,7 @@ class DevicenDisplay(DeviceUnreal):
                 args.append("-statnamedevents")
 
             # if bookmarks are enabled, also enable the vblank monitoring thread
-            if 'bookmark' in traces.split(','):
+            if any(bmtrace in traces.split(',') for bmtrace in ('bookmark','default')) :
                 dp_cvars.append('nDisplay.sync.diag.VBlankMonitoring=1')
 
         # Always tell Chaos to be deterministic
