Switchboard: sbl_helper: Fix handling of `p4 print` results returned via multiple `text` records.

diff --git a/Engine/Plugins/VirtualProduction/Switchboard/Source/Switchboard/sbl_helper.py b/Engine/Plugins/VirtualProduction/Switchboard/Source/Switchboard/sbl_helper.py
--- a/Engine/Plugins/VirtualProduction/Switchboard/Source/Switchboard/sbl_helper.py
+++ b/Engine/Plugins/VirtualProduction/Switchboard/Source/Switchboard/sbl_helper.py
@@ -138,36 +138,52 @@ def p4_have(
     return result
 
 
+@dataclass
+class P4PrintResult:
+    stat: dict
+    contents: Optional[bytearray]
+
+    @property
+    def is_valid(self) -> bool:
+        try:
+            return (self.stat[b'code'] == b'stat') and (self.contents is not None)
+        except KeyError:
+            return False
+
+    @property
+    def text(self) -> Optional[str]:
+        if self.contents is not None:
+            return self.contents.decode()
+        else:
+            return None
+
+
 def p4_print(
     pathspecs: List[str],
     *,
     user: Optional[str] = None,
     client: Optional[str] = None,
-) -> List[Tuple[Optional[str], dict]]:
+) -> List[P4PrintResult]:
     '''
-    Returns a list of (fileContents, statRecord).
+    Returns a `P4PrintResult` for each path in `pathspecs`.
     If the file does not exist, or is deleted at the specified revision,
     fileContents will be `None`.
     '''
-    results: List[Tuple[Optional[str], dict]] = []
+    results: List[P4PrintResult] = []
     records = p4_get_records('print', pathspecs, include_info=True,
                              include_error=True, user=user, client=client)
 
-    # For each file, a stat record is returned, followed by a text record.
-    # In the case of an error (e.g. nonexistent file), there is no text record.
-    last_stat_record = {}
+    # For each file, a stat record is returned, then zero or more text records.
+    # Cases where there's no text record include errors or deleted files.
     for record in records:
-        if record[b'code'] == b'stat':
-            if b'delete' in record[b'action']:  # 'deleted' or 'move/delete'
-                results.append((None, record))
-            else:
-                last_stat_record = record
+        if (record[b'code'] == b'stat') or (record[b'code'] == b'error'):
+            results.append(P4PrintResult(record, None))
             continue
         elif record[b'code'] == b'text':
-            results.append((record[b'data'].decode(), last_stat_record))
-            continue
-        elif record[b'code'] == b'error':
-            results.append((None, record))
+            current_result = results[-1]
+            if current_result.contents is None:
+                current_result.contents = bytearray()
+            current_result.contents.extend(record[b'data'])
             continue
 
         logging.warning(f'p4_print(): Unhandled record: {record}')
@@ -870,10 +886,13 @@ class SbListenerHelper:
             [f'{path}@{self.sync_engine_cl}'
              for path in (epicint_local_path, buildver_local_path)],
             user=self.p4user, client=self.p4client)
-        epicint_depot_text, _ = prints[0]
-        ret_info.depot_text, buildver_meta = prints[1]
+        epicint_print_result = prints[0]
+        buildver_print_result = prints[1]
+
+        ret_info.depot_text = buildver_print_result.text  # May be None
 
-        if ret_info.depot_text is not None:
+        if buildver_print_result.is_valid:
+            buildver_meta = buildver_print_result.stat
             ret_info.depot_path = buildver_meta[b'depotFile'].decode()
 
             # Code CL determination compatible with precompiled binaries.
@@ -904,7 +923,7 @@ class SbListenerHelper:
                 '/Engine/Build/Build.version', '')
             updates = {
                 'Changelist': self.sync_engine_cl,
-                'IsLicenseeVersion': 1 if epicint_depot_text is None else 0,
+                'IsLicenseeVersion': 0 if epicint_print_result.is_valid else 1,
                 'BranchName': branch_name.replace('/', '+'),
             }
 
