Switchboard: Added sm5 and sm6 variants when selecting the api used to launch ndisplay.

This is to mitigate certain cases some clients ran into that were unknowingly running with sm5 and things were not working as expected.

#skipci

diff --git a/Engine/Plugins/VirtualProduction/Switchboard/Source/Switchboard/switchboard/devices/ndisplay/plugin_ndisplay.py b/Engine/Plugins/VirtualProduction/Switchboard/Source/Switchboard/switchboard/devices/ndisplay/plugin_ndisplay.py
--- a/Engine/Plugins/VirtualProduction/Switchboard/Source/Switchboard/switchboard/devices/ndisplay/plugin_ndisplay.py
+++ b/Engine/Plugins/VirtualProduction/Switchboard/Source/Switchboard/switchboard/devices/ndisplay/plugin_ndisplay.py
@@ -297,7 +297,15 @@ class DevicenDisplay(DeviceUnreal):
             attr_name="render_api",
             nice_name="Render API",
             value="dx12",
-            possible_values=["dx11", "dx12", "vulkan"],
+            possible_values=[
+                "dx11",
+                "dx11 -sm5",
+                "dx11 -sm6",
+                "dx12", 
+                "dx12 -sm5", 
+                "dx12 -sm6", 
+                "vulkan"
+            ],
         ),
         'render_mode': OptionSetting(
             attr_name="render_mode",
