Always pass VPRoles so that set by the project do not cause issues.

diff --git a/Engine/Plugins/VirtualProduction/Switchboard/Source/Switchboard/switchboard/devices/unreal/plugin_unreal.py b/Engine/Plugins/VirtualProduction/Switchboard/Source/Switchboard/switchboard/devices/unreal/plugin_unreal.py
--- a/Engine/Plugins/VirtualProduction/Switchboard/Source/Switchboard/switchboard/devices/unreal/plugin_unreal.py
+++ b/Engine/Plugins/VirtualProduction/Switchboard/Source/Switchboard/switchboard/devices/unreal/plugin_unreal.py
@@ -1671,8 +1671,7 @@ class DeviceUnreal(Device):
 
         (supported_roles, unsupported_roles) = self.get_vproles()
 
-        if supported_roles:
-            command_line_args += ' -VPRole=' + '|'.join(supported_roles)
+        command_line_args += ' -VPRole=' + '|'.join(supported_roles)
         if unsupported_roles:
             LOGGER.error(
                 f"{self.name}: Omitted unsupported roles: "
