If any devices have a VPRole set then all devices need to pass in a role value even if it is null.  This is to prevent conflicts for Virtual Camera / Multi-user setups where the project settings cause other nodes to participate even if they are not suppose to.

diff --git a/Engine/Plugins/VirtualProduction/Switchboard/Source/Switchboard/switchboard/devices/unreal/plugin_unreal.py b/Engine/Plugins/VirtualProduction/Switchboard/Source/Switchboard/switchboard/devices/unreal/plugin_unreal.py
--- a/Engine/Plugins/VirtualProduction/Switchboard/Source/Switchboard/switchboard/devices/unreal/plugin_unreal.py
+++ b/Engine/Plugins/VirtualProduction/Switchboard/Source/Switchboard/switchboard/devices/unreal/plugin_unreal.py
@@ -1528,6 +1528,14 @@ class DeviceUnreal(Device):
         return CONFIG.engine_exe_path(
             CONFIG.ENGINE_DIR.get_value(self.name), self.executable_filename)
 
+    @staticmethod
+    def any_devices_have_roles():
+        device_list = DeviceUnreal.active_unreal_devices
+        for device in device_list:
+            if len(device.get_vproles()[0]):
+                return True
+        return False
+
     def get_vproles(self):
         ''' Gets selected vp roles that are also present in the ini file
         Also returns any selected vp roles that are not in the ini file.
@@ -1671,7 +1679,9 @@ class DeviceUnreal(Device):
 
         (supported_roles, unsupported_roles) = self.get_vproles()
 
-        command_line_args += ' -VPRole=' + '|'.join(supported_roles)
+        if supported_roles or DeviceUnreal.any_devices_have_roles():
+            command_line_args += ' -VPRole=' + '|'.join(supported_roles)
+
         if unsupported_roles:
             LOGGER.error(
                 f"{self.name}: Omitted unsupported roles: "
