Switchboard: It now defaults to put -DisableFrameTraceCapture in the command line when launching nDisplay nodes.

This is because RenderDoc is now enabled by default in the Engine, and when launched with -unattended (like nDisplay), then it is effectively turned on.

When RenderDoc is turned on, then Unreal calls FDynamicRHI::EnableIdealGPUCaptureOptions(true), which disables important optimizations, and fps drops dramatically (about 50% in a representative project).

When trying to use RenderDoc, you must now disable this flag in Switchboard:
  Disable Frame Trace Capture

diff --git a/Engine/Plugins/VirtualProduction/Switchboard/Source/Switchboard/switchboard/devices/ndisplay/plugin_ndisplay.py b/Engine/Plugins/VirtualProduction/Switchboard/Source/Switchboard/switchboard/devices/ndisplay/plugin_ndisplay.py
--- a/Engine/Plugins/VirtualProduction/Switchboard/Source/Switchboard/switchboard/devices/ndisplay/plugin_ndisplay.py
+++ b/Engine/Plugins/VirtualProduction/Switchboard/Source/Switchboard/switchboard/devices/ndisplay/plugin_ndisplay.py
@@ -459,6 +459,14 @@ class DevicenDisplay(DeviceUnreal):
             tool_tip=(
                 'Adds the selected LiveLink preset to the command line \n')
         ),
+        'disable_frame_trace_capture': BoolSetting(
+            attr_name='disable_frame_trace_capture',
+            nice_name='Disable Frame Trace Capture',
+            value=True,
+            tool_tip=(
+                'Prevents RenderDoc from loading, so that Unreal Engine does not disable certain rendering optimizations '
+                'such as parallel algorithms. Uncheck this setting to allow RenderDoc captures. \n')
+        ),
     }
 
     ndisplay_monitor_ui = None
@@ -749,23 +757,25 @@ class DevicenDisplay(DeviceUnreal):
             if DevicenDisplay.csettings['disable_all_screen_messages'].get_value()
             else '')
 
+        # Effectively disable RenderDoc
+        disable_frame_trace_capture = (
+            '-DisableFrameTraceCapture'
+            if DevicenDisplay.csettings['disable_frame_trace_capture'].get_value()
+            else '')
+
         # fill in fixed arguments
         args = [
             f'"{uproject}"',
             "-game",                      # render nodes run in -game
             f'{map_name}',                # map to open
-            "-messaging",                 # enables messaging, needed for
-                                          # MultiUser
+            "-messaging",                 # enables messaging, needed for MultiUser
             "-dc_cluster",                # this is a cluster node
             "-nosplash",                  # avoids splash screen
             "-fixedseed",                 # for determinism
             "-NoVerifyGC",                # improves performance
             "-noxrstereo",                # avoids a conflict with steam/oculus
-            "-xrtrackingonly",            # allows multi-UE SteamVR for
-                                          # trackers (but disallows sending
-                                          # frames)
-            "-RemoteControlIsHeadless",   # avoids notification window when
-                                          # using RemoteControlWebUI
+            "-xrtrackingonly",            # allows multi-UE SteamVR for trackers (but disallows sending frames)
+            "-RemoteControlIsHeadless",   # avoids notification window when using RemoteControlWebUI
             f'{additional_args}',         # specified in settings
             f'{vproles}',                 # VP roles for this instance
             f'{friendly_name}',           # Stage Friendly Name
@@ -776,8 +786,7 @@ class DevicenDisplay(DeviceUnreal):
             f'{render_mode}',             # mono/...
             f'{use_all_cores}',           # -useallavailablecores
             f'{no_texture_streaming}',    # -notexturestreaming
-            f'-dc_node={self.name}',      # name of this node in the nDisplay
-                                          # cluster
+            f'-dc_node={self.name}',      # name of this node in the nDisplay cluster
             f'Log={self.log_filename}',   # log file
             f'{ini_engine}',              # Engine ini injections
             f'{ini_game}',                # Game ini injections
@@ -785,6 +794,7 @@ class DevicenDisplay(DeviceUnreal):
             f'{unattended}',              # -unattended
             f'{no_screen_messages}',      # -NoScreenMessages
             f'{disable_ensures}',         # -handleensurepercent=0
+            f'{disable_frame_trace_capture}', # -DisableFrameTraceCapture
             f'{udpm_transport_multi}',    # -UDPMESSAGING_TRANSPORT_MULTICAST=
             f'{udpm_transport_unicast}',  # -UDPMESSAGING_TRANSPORT_UNICAST=
             f'{udpm_transport_static}',   # -UDPMESSAGING_TRANSPORT_STATIC=
