nDisplay: Added an option to set the Graphics Adapter used to launch a given nDisplay node.

This is particularly useful when in non sli multi-gpu system.

diff --git a/Engine/Plugins/Runtime/nDisplay/Source/DisplayClusterConfiguration/Private/Formats/JSON500/DisplayClusterConfigurationJsonParser_500.cpp b/Engine/Plugins/Runtime/nDisplay/Source/DisplayClusterConfiguration/Private/Formats/JSON500/DisplayClusterConfigurationJsonParser_500.cpp
--- a/Engine/Plugins/Runtime/nDisplay/Source/DisplayClusterConfiguration/Private/Formats/JSON500/DisplayClusterConfigurationJsonParser_500.cpp
+++ b/Engine/Plugins/Runtime/nDisplay/Source/DisplayClusterConfiguration/Private/Formats/JSON500/DisplayClusterConfigurationJsonParser_500.cpp
@@ -222,6 +222,7 @@ namespace JSON500
 				Node->bIsSoundEnabled = CfgNode.Value.Sound;
 				Node->bIsFullscreen = CfgNode.Value.FullScreen;
 				Node->bRenderHeadless = CfgNode.Value.RenderHeadless;
+				Node->GraphicsAdapter = CfgNode.Value.GraphicsAdapter;
 				Node->bEnableTextureShare = CfgNode.Value.TextureShare;
 				Node->WindowRect = FDisplayClusterConfigurationRectangle(CfgNode.Value.Window.X, CfgNode.Value.Window.Y, CfgNode.Value.Window.W, CfgNode.Value.Window.H);
 
@@ -415,6 +416,7 @@ namespace JSON500
 				Node.Sound = CfgNode.Value->bIsSoundEnabled;
 				Node.FullScreen = CfgNode.Value->bIsFullscreen;
 				Node.RenderHeadless = CfgNode.Value->bRenderHeadless;
+				Node.GraphicsAdapter = CfgNode.Value->GraphicsAdapter;
 				Node.TextureShare = CfgNode.Value->bEnableTextureShare;
 				Node.Window = FDisplayClusterConfigurationJsonRectangle_500(CfgNode.Value->WindowRect.X, CfgNode.Value->WindowRect.Y, CfgNode.Value->WindowRect.W, CfgNode.Value->WindowRect.H);
 
diff --git a/Engine/Plugins/Runtime/nDisplay/Source/DisplayClusterConfiguration/Private/Formats/JSON500/DisplayClusterConfigurationJsonTypes_500.h b/Engine/Plugins/Runtime/nDisplay/Source/DisplayClusterConfiguration/Private/Formats/JSON500/DisplayClusterConfigurationJsonTypes_500.h
--- a/Engine/Plugins/Runtime/nDisplay/Source/DisplayClusterConfiguration/Private/Formats/JSON500/DisplayClusterConfigurationJsonTypes_500.h
+++ b/Engine/Plugins/Runtime/nDisplay/Source/DisplayClusterConfiguration/Private/Formats/JSON500/DisplayClusterConfigurationJsonTypes_500.h
@@ -501,6 +501,9 @@ public:
 	UPROPERTY()
 	bool RenderHeadless = false;
 
+	UPROPERTY()
+	int32 GraphicsAdapter = -1;
+
 	UPROPERTY()
 	bool TextureShare = false;
 
diff --git a/Engine/Plugins/Runtime/nDisplay/Source/DisplayClusterConfiguration/Public/DisplayClusterConfigurationTypes.h b/Engine/Plugins/Runtime/nDisplay/Source/DisplayClusterConfiguration/Public/DisplayClusterConfigurationTypes.h
--- a/Engine/Plugins/Runtime/nDisplay/Source/DisplayClusterConfiguration/Public/DisplayClusterConfigurationTypes.h
+++ b/Engine/Plugins/Runtime/nDisplay/Source/DisplayClusterConfiguration/Public/DisplayClusterConfigurationTypes.h
@@ -338,6 +338,10 @@ public:
 	UPROPERTY(EditAnywhere, BlueprintReadonly, Category = "Configuration", meta = (DisplayName = "Headless Rendering"))
 	bool bRenderHeadless = false;
 
+	/** Hint for setting the r.GraphicsAdapter CVar when launching this cluster node. Note that this is distinct from the GPU Node Indices assigned to viewports. */
+	UPROPERTY(EditAnywhere, BlueprintReadonly, Category = "Configuration", meta = (DisplayName = "Graphics Adapter"))
+	int32 GraphicsAdapter = -1;
+
 	/** Enables texture sharing for this cluster node */
 	UPROPERTY(EditAnywhere, BlueprintReadonly, Category = "Configuration", meta = (DisplayName = "Enable Texture Share"))
 	bool bEnableTextureShare = false;
diff --git a/Engine/Plugins/VirtualProduction/Switchboard/Source/Switchboard/switchboard/devices/ndisplay/plugin_ndisplay.py b/Engine/Plugins/VirtualProduction/Switchboard/Source/Switchboard/switchboard/devices/ndisplay/plugin_ndisplay.py
--- a/Engine/Plugins/VirtualProduction/Switchboard/Source/Switchboard/switchboard/devices/ndisplay/plugin_ndisplay.py
+++ b/Engine/Plugins/VirtualProduction/Switchboard/Source/Switchboard/switchboard/devices/ndisplay/plugin_ndisplay.py
@@ -16,7 +16,7 @@ from switchboard import message_protocol
 from switchboard import switchboard_utils as sb_utils
 from switchboard import switchboard_widgets as sb_widgets
 from switchboard import switchboard_dialog as sb_dialog
-from switchboard.config import CONFIG, BoolSetting, FilePathSetting, \
+from switchboard.config import CONFIG, BoolSetting, IntSetting, FilePathSetting, \
     LoggingSetting, OptionSetting, Setting, StringSetting, SETTINGS, \
     StringListSetting, migrate_comma_separated_string_to_list
 from switchboard.devices.device_widget_base import AddDeviceDialog
@@ -467,6 +467,17 @@ class DevicenDisplay(DeviceUnreal):
                 'Prevents RenderDoc from loading, so that Unreal Engine does not disable certain rendering optimizations '
                 'such as parallel algorithms. Uncheck this setting to allow RenderDoc captures. \n')
         ),
+        'graphics_adapter': OptionSetting(
+            attr_name="graphics_adapter",
+            nice_name="Graphics Adapter",
+            value='Config',
+            possible_values=['Config'] + list(range(8)),
+            tool_tip=(
+                "Select which graphics adapter to use (will set r.GraphicsAdapter early CVar) \n"
+                "- 'Config' : Use the setting in the nDisplay config file \n"
+                "- 0, 1, .. : The specified gpu index \n"
+            ),
+        ),
     }
 
     ndisplay_monitor_ui = None
@@ -507,6 +518,12 @@ class DevicenDisplay(DeviceUnreal):
                 value=kwargs.get("renderHeadless", False),
                 show_ui=False
             ),
+            'graphics_adapter' : IntSetting(
+                attr_name="graphics_adapter",
+                nice_name="Graphics Adapter",
+                value=kwargs.get("graphics_adapter", -1),
+                show_ui=False
+            ),
         }
 
         self.render_mode_cmdline_opts = {
@@ -602,6 +619,7 @@ class DevicenDisplay(DeviceUnreal):
             DevicenDisplay.csettings['udpmessaging_unicast_endpoint'],
             DevicenDisplay.csettings['udpmessaging_extra_static_endpoints'],
             DevicenDisplay.csettings['livelink_preset'],
+            DevicenDisplay.csettings['graphics_adapter'],
             CONFIG.ENGINE_DIR,
             CONFIG.SOURCE_CONTROL_WORKSPACE,
             CONFIG.UPROJECT_PATH,
@@ -889,6 +907,14 @@ class DevicenDisplay(DeviceUnreal):
             if any(bmtrace in traces.split(',') for bmtrace in ('bookmark','default')) :
                 dp_cvars.append('nDisplay.sync.diag.VBlankMonitoring=1')
 
+        # Graphics Adapter
+        graphics_adapter = DevicenDisplay.csettings["graphics_adapter"].get_value(self.name)
+        if graphics_adapter == 'Config':
+            graphics_adapter = self.settings['graphics_adapter'].get_value()
+
+        if graphics_adapter >= 0:
+            dp_cvars.append(f'r.GraphicsAdapter={graphics_adapter}')
+
         # Always tell Chaos to be deterministic
         dp_cvars.append('p.Chaos.Solver.Deterministic=1')
 
@@ -1087,6 +1113,11 @@ class DevicenDisplay(DeviceUnreal):
             if (device.name == DevicenDisplay.csettings['primary_device_name'].get_value()):
                 primaryNodeId = device.name
 
+            # override GraphicsAdapter
+            graphics_adapter = DevicenDisplay.csettings['graphics_adapter'].get_value(device.name)
+            if graphics_adapter != 'Config':
+                node['graphicsAdapter'] = int(graphics_adapter)
+
         data['nDisplay']['cluster']['nodes'] = activenodes
         data['nDisplay']['cluster'][primaryNodeKey]['id'] = primaryNodeId
 
@@ -1155,6 +1186,7 @@ class DevicenDisplay(DeviceUnreal):
             # Note the capital 'S', 'H'...
             kwargs["fullscreen"] = bool(cnode.get('fullScreen', False))
             kwargs["headless"] = bool(cnode.get('renderHeadless', False))
+            kwargs["graphics_adapter"] = int(cnode.get('graphicsAdapter', -1))
 
             primary = True if primaryNode['id'] == name else False
 
@@ -1216,6 +1248,8 @@ class DevicenDisplay(DeviceUnreal):
             menode['kwargs'].get("fullscreen", False))
         self.settings['headless'].update_value(
             menode['kwargs'].get("headless", False))
+        self.settings['graphics_adapter'].update_value(
+            menode['kwargs'].get("graphics_adapter", -1))
 
     @classmethod
     def uasset_path_from_object_path(
