Simplify RenderDoc start up logic to simplify behavior on unattended runs.
This comes from a recent change where RenderDoc plugin is always loaded, but relies on command line to determine if it registers for use.
This led to some unattended processes enabling RenderDoc without wanting it and seeing a perf regression.
Now even unattended processes need -attachRenderDoc on the command line.
Changed automation scripts so that -attachRenderDoc will be used for automation tests if we set AttachRenderDoc=true (it's false by default).
Backed out changes in ndisplay that were working around the perf regression.

diff --git a/Engine/Plugins/Developer/RenderDocPlugin/Source/RenderDocPlugin/Private/RenderDocPluginLoader.cpp b/Engine/Plugins/Developer/RenderDocPlugin/Source/RenderDocPlugin/Private/RenderDocPluginLoader.cpp
--- a/Engine/Plugins/Developer/RenderDocPlugin/Source/RenderDocPlugin/Private/RenderDocPluginLoader.cpp
+++ b/Engine/Plugins/Developer/RenderDocPlugin/Source/RenderDocPlugin/Private/RenderDocPluginLoader.cpp
@@ -195,28 +195,9 @@ void FRenderDocPluginLoader::Initialize()
 	UE::ConfigUtilities::ApplyCVarSettingsFromIni(TEXT("/Script/RenderDocPlugin.RenderDocPluginSettings"), *GEngineIni, ECVF_SetByProjectSetting);
 
 	const bool bCapture = FParse::Param(FCommandLine::Get(), TEXT("AttachRenderDoc")) || (CVarRenderDocAutoAttach.GetValueOnAnyThread() != 0);
-
-	const bool bUnattended = FApp::IsUnattended();
-	IConsoleVariable* CVarAutomationAllowFrameTraceCapture = IConsoleManager::Get().FindConsoleVariable(TEXT("AutomationAllowFrameTraceCapture"), false);
-	const bool bUnattendedCapture = bUnattended && CVarAutomationAllowFrameTraceCapture && CVarAutomationAllowFrameTraceCapture->GetInt() != 0;
-
-	if (!bCapture && !bUnattendedCapture)
-	{
-		if (bUnattended)
-		{
-			UE_LOG(RenderDocPlugin, Display, TEXT("RenderDoc plugin will not be loaded because AutomationAllowFrameTraceCapture cvar is set to 0."));
-		}
-		else
-		{
-			UE_LOG(RenderDocPlugin, Display, TEXT("RenderDoc plugin will not be loaded. Use '-AttachRenderDoc' on the cmd line or enable 'renderdoc.AutoAttach' in the plugin settings."));
-		}
-		return;
-	}
-
-	const bool bDisableFrameTraceCapture = FParse::Param(FCommandLine::Get(), TEXT("DisableFrameTraceCapture"));
-	if (bDisableFrameTraceCapture)
+	if (!bCapture)
 	{
-		UE_LOG(RenderDocPlugin, Display, TEXT("RenderDoc plugin will not be loaded because -DisableFrameTraceCapture cmd line flag."));
+		UE_LOG(RenderDocPlugin, Display, TEXT("RenderDoc plugin will not be loaded. Use '-AttachRenderDoc' on the cmd line or enable 'renderdoc.AutoAttach' in the plugin settings."));
 		return;
 	}
 
diff --git a/Engine/Plugins/VirtualProduction/Switchboard/Source/Switchboard/switchboard/devices/ndisplay/plugin_ndisplay.py b/Engine/Plugins/VirtualProduction/Switchboard/Source/Switchboard/switchboard/devices/ndisplay/plugin_ndisplay.py
--- a/Engine/Plugins/VirtualProduction/Switchboard/Source/Switchboard/switchboard/devices/ndisplay/plugin_ndisplay.py
+++ b/Engine/Plugins/VirtualProduction/Switchboard/Source/Switchboard/switchboard/devices/ndisplay/plugin_ndisplay.py
@@ -459,14 +459,6 @@ class DevicenDisplay(DeviceUnreal):
             tool_tip=(
                 'Adds the selected LiveLink preset to the command line \n')
         ),
-        'disable_frame_trace_capture': BoolSetting(
-            attr_name='disable_frame_trace_capture',
-            nice_name='Disable Frame Trace Capture',
-            value=True,
-            tool_tip=(
-                'Prevents RenderDoc from loading, so that Unreal Engine does not disable certain rendering optimizations '
-                'such as parallel algorithms. Uncheck this setting to allow RenderDoc captures. \n')
-        ),
         'graphics_adapter': OptionSetting(
             attr_name="graphics_adapter",
             nice_name="Graphics Adapter",
@@ -775,12 +767,6 @@ class DevicenDisplay(DeviceUnreal):
             if DevicenDisplay.csettings['disable_all_screen_messages'].get_value()
             else '')
 
-        # Effectively disable RenderDoc
-        disable_frame_trace_capture = (
-            '-DisableFrameTraceCapture'
-            if DevicenDisplay.csettings['disable_frame_trace_capture'].get_value()
-            else '')
-
         # fill in fixed arguments
         args = [
             f'"{uproject}"',
@@ -812,7 +798,6 @@ class DevicenDisplay(DeviceUnreal):
             f'{unattended}',              # -unattended
             f'{no_screen_messages}',      # -NoScreenMessages
             f'{disable_ensures}',         # -handleensurepercent=0
-            f'{disable_frame_trace_capture}', # -DisableFrameTraceCapture
             f'{udpm_transport_multi}',    # -UDPMESSAGING_TRANSPORT_MULTICAST=
             f'{udpm_transport_unicast}',  # -UDPMESSAGING_TRANSPORT_UNICAST=
             f'{udpm_transport_static}',   # -UDPMESSAGING_TRANSPORT_STATIC=
diff --git a/Engine/Source/Programs/AutomationTool/Gauntlet/Unreal/Automation/UE.Automation.cs b/Engine/Source/Programs/AutomationTool/Gauntlet/Unreal/Automation/UE.Automation.cs
--- a/Engine/Source/Programs/AutomationTool/Gauntlet/Unreal/Automation/UE.Automation.cs
+++ b/Engine/Source/Programs/AutomationTool/Gauntlet/Unreal/Automation/UE.Automation.cs
@@ -77,10 +77,10 @@ namespace UE
 		public bool D3DDebug = false;
 
 		/// <summary>
-		/// Disable capturing frame trace for image based tests
+		/// Attach to RenderDoc to capture frame traces for image based tests
 		/// </summary>
 		[AutoParam]
-		public bool DisableFrameTraceCapture = true;
+		public bool AttachRenderDoc = false;
 
 		/// <summary>
 		/// Filter or groups of tests to apply
@@ -239,14 +239,9 @@ namespace UE
 			if ((ConfigRole.RoleType.IsEditor() && HasNoOtherRole) || ConfigRole.RoleType.IsClient())
 			{
 				// These are flags that are required only on the main role that is going to execute the tests. ie raytracing is required only on the client or if it is an editor test.
-				if (DisableFrameTraceCapture || RayTracing)
-				{
-					AppConfig.CommandLine += " -DisableFrameTraceCapture";
-				}
-
 				if (RayTracing)
 				{
-					AppConfig.CommandLine += " -dpcvars=r.RayTracing=1,r.SkinCache.CompileShaders=1,r.Lumen.HardwareRayTracing=1,AutomationAllowFrameTraceCapture=0";
+					AppConfig.CommandLine += " -dpcvars=r.RayTracing=1,r.SkinCache.CompileShaders=1,r.Lumen.HardwareRayTracing=1";
 				}
 				else
 				{
@@ -256,6 +251,11 @@ namespace UE
 				// Options specific to windows
 				if (ConfigRole.Platform != null && ((UnrealTargetPlatform)ConfigRole.Platform).IsInGroup(UnrealPlatformGroup.Windows))
 				{
+					if (AttachRenderDoc && !RayTracing)
+					{
+						AppConfig.CommandLine += " -attachRenderDoc";
+					}
+
 					if (PreferNvidia)
 					{
 						AppConfig.CommandLine += " -preferNvidia";
diff --git a/Engine/Source/Runtime/Engine/Private/Tests/AutomationCommon.cpp b/Engine/Source/Runtime/Engine/Private/Tests/AutomationCommon.cpp
--- a/Engine/Source/Runtime/Engine/Private/Tests/AutomationCommon.cpp
+++ b/Engine/Source/Runtime/Engine/Private/Tests/AutomationCommon.cpp
@@ -158,8 +158,7 @@ namespace AutomationCommon
 	{
 		TArray<uint8> FrameTrace;
 
-		bool bDisableFrameTraceCapture = FParse::Param(FCommandLine::Get(), TEXT("DisableFrameTraceCapture"));
-		if (!bDisableFrameTraceCapture && CVarAutomationAllowFrameTraceCapture.GetValueOnGameThread() != 0 && IRenderCaptureProvider::IsAvailable())
+		if (CVarAutomationAllowFrameTraceCapture.GetValueOnGameThread() != 0 && IRenderCaptureProvider::IsAvailable())
 		{
 			const FString MapAndTest = MapOrContext / FPaths::MakeValidFileName(TestName, TEXT('_'));
 			FString ScreenshotName = GetScreenshotName(MapAndTest);
