Switchboard: SearchableCombo boxes now by default only mousewheel when in focus.

The LiveLink preset combo now needs tab or click to focus. This helps users to not accidentally change the LL preset when scrolling down the settings dialog.

diff --git a/Engine/Plugins/VirtualProduction/Switchboard/Source/Switchboard/switchboard/devices/unreal/plugin_unreal.py b/Engine/Plugins/VirtualProduction/Switchboard/Source/Switchboard/switchboard/devices/unreal/plugin_unreal.py
--- a/Engine/Plugins/VirtualProduction/Switchboard/Source/Switchboard/switchboard/devices/unreal/plugin_unreal.py
+++ b/Engine/Plugins/VirtualProduction/Switchboard/Source/Switchboard/switchboard/devices/unreal/plugin_unreal.py
@@ -299,6 +299,8 @@ class LiveLinkPresetSetting(Setting):
 
         combo = sb_widgets.SearchableComboBox(None)
         combo.setSizePolicy(QtWidgets.QSizePolicy.Expanding, QtWidgets.QSizePolicy.Preferred)
+        combo.setFocusPolicy(QtCore.Qt.FocusPolicy.StrongFocus)
+
         self._update_combo_items(combo, override_device_name)
 
         combo.currentTextChanged.connect(
diff --git a/Engine/Plugins/VirtualProduction/Switchboard/Source/Switchboard/switchboard/switchboard_widgets.py b/Engine/Plugins/VirtualProduction/Switchboard/Source/Switchboard/switchboard/switchboard_widgets.py
--- a/Engine/Plugins/VirtualProduction/Switchboard/Source/Switchboard/switchboard/switchboard_widgets.py
+++ b/Engine/Plugins/VirtualProduction/Switchboard/Source/Switchboard/switchboard/switchboard_widgets.py
@@ -361,6 +361,7 @@ class SearchableComboBox(QtWidgets.QComboBox):
         self.search_line.setPlaceholderText("Search")
 
         self.search_line.textChanged.connect(self.text_changed)
+        self.wheel_only_on_focus = True
 
     def showPopup(self):
         super().showPopup()
@@ -388,6 +389,14 @@ class SearchableComboBox(QtWidgets.QComboBox):
             else:
                 self.view().setRowHidden(self.findText(item), True)
 
+    def wheelEvent(self, event) -> None:
+
+        if self.wheel_only_on_focus:
+            if self.hasFocus():
+                super().wheelEvent(event)
+        else:
+            super().wheelEvent(event)
+
 
 class HostnameValidator(QtGui.QValidator):
     INVALID_CHARS_RE = re.compile(r"[^A-Z.\d-]", re.IGNORECASE)
