Switchboard: Fix for slate multi-user server being unable to launch if the (unquoted) path contains a space.

diff --git a/Engine/Plugins/VirtualProduction/Switchboard/Source/Switchboard/switchboard/switchboard_application.py b/Engine/Plugins/VirtualProduction/Switchboard/Source/Switchboard/switchboard/switchboard_application.py
--- a/Engine/Plugins/VirtualProduction/Switchboard/Source/Switchboard/switchboard/switchboard_application.py
+++ b/Engine/Plugins/VirtualProduction/Switchboard/Source/Switchboard/switchboard/switchboard_application.py
@@ -139,11 +139,11 @@ class MultiUserApplication:
             cmdline = ''
             if sys.platform.startswith('win'):
                 if CONFIG.MUSERVER_SLATE_MODE.get_value():
-                    cmdline = f'{self.exe_path()}'
+                    cmdline = f'"{self.exe_path()}"'
                 else:
                     cmdline = f'start "Multi User Server" "{self.exe_path()}"'
             else:
-                cmdline = f'{self.exe_path()}'
+                cmdline = f'"{self.exe_path()}"'
 
             cmdline += f' -CONCERTSERVER="{CONFIG.MUSERVER_SERVER_NAME.get_value()}"'
             cmdline += f' {CONFIG.MUSERVER_COMMAND_LINE_ARGUMENTS.get_value()}'
