Switchboard: Fixed bug with graphics adapter override.

diff --git a/Engine/Plugins/VirtualProduction/Switchboard/Source/Switchboard/switchboard/devices/ndisplay/plugin_ndisplay.py b/Engine/Plugins/VirtualProduction/Switchboard/Source/Switchboard/switchboard/devices/ndisplay/plugin_ndisplay.py
--- a/Engine/Plugins/VirtualProduction/Switchboard/Source/Switchboard/switchboard/devices/ndisplay/plugin_ndisplay.py
+++ b/Engine/Plugins/VirtualProduction/Switchboard/Source/Switchboard/switchboard/devices/ndisplay/plugin_ndisplay.py
@@ -510,9 +510,9 @@ class DevicenDisplay(DeviceUnreal):
                 value=kwargs.get("renderHeadless", False),
                 show_ui=False
             ),
-            'graphics_adapter' : IntSetting(
-                attr_name="graphics_adapter",
-                nice_name="Graphics Adapter",
+            'config_graphics_adapter' : IntSetting(
+                attr_name="config_graphics_adapter",
+                nice_name="Config Graphics Adapter",
                 value=kwargs.get("graphics_adapter", -1),
                 show_ui=False
             ),
@@ -895,7 +895,7 @@ class DevicenDisplay(DeviceUnreal):
         # Graphics Adapter
         graphics_adapter = DevicenDisplay.csettings["graphics_adapter"].get_value(self.name)
         if graphics_adapter == 'Config':
-            graphics_adapter = self.settings['graphics_adapter'].get_value()
+            graphics_adapter = self.settings["config_graphics_adapter"].get_value(self.name)
 
         if graphics_adapter >= 0:
             dp_cvars.append(f'r.GraphicsAdapter={graphics_adapter}')
@@ -1171,7 +1171,7 @@ class DevicenDisplay(DeviceUnreal):
             # Note the capital 'S', 'H'...
             kwargs["fullscreen"] = bool(cnode.get('fullScreen', False))
             kwargs["headless"] = bool(cnode.get('renderHeadless', False))
-            kwargs["graphics_adapter"] = int(cnode.get('graphicsAdapter', -1))
+            kwargs["config_graphics_adapter"] = int(cnode.get('graphicsAdapter', -1))
 
             primary = True if primaryNode['id'] == name else False
 
@@ -1233,8 +1233,8 @@ class DevicenDisplay(DeviceUnreal):
             menode['kwargs'].get("fullscreen", False))
         self.settings['headless'].update_value(
             menode['kwargs'].get("headless", False))
-        self.settings['graphics_adapter'].update_value(
-            menode['kwargs'].get("graphics_adapter", -1))
+        self.settings['config_graphics_adapter'].update_value(
+            menode['kwargs'].get("config_graphics_adapter", -1))
 
     @classmethod
     def uasset_path_from_object_path(
