Switchboard: Added a way to set the media profile in switchboard, including per device overrides.

diff --git a/Engine/Plugins/VirtualProduction/Switchboard/Source/Switchboard/switchboard/devices/ndisplay/plugin_ndisplay.py b/Engine/Plugins/VirtualProduction/Switchboard/Source/Switchboard/switchboard/devices/ndisplay/plugin_ndisplay.py
--- a/Engine/Plugins/VirtualProduction/Switchboard/Source/Switchboard/switchboard/devices/ndisplay/plugin_ndisplay.py
+++ b/Engine/Plugins/VirtualProduction/Switchboard/Source/Switchboard/switchboard/devices/ndisplay/plugin_ndisplay.py
@@ -21,7 +21,7 @@ from switchboard.config import CONFIG, BoolSetting, IntSetting, FilePathSetting,
     StringListSetting, migrate_comma_separated_string_to_list
 from switchboard.devices.device_widget_base import AddDeviceDialog
 from switchboard.devices.unreal.plugin_unreal import DeviceUnreal, \
-    DeviceWidgetUnreal, LiveLinkPresetSetting
+    DeviceWidgetUnreal, LiveLinkPresetSetting, MediaProfileSetting
 from switchboard.devices.unreal.uassetparser import UassetParser
 from switchboard.switchboard_logging import LOGGER
 
@@ -470,6 +470,12 @@ class DevicenDisplay(DeviceUnreal):
                 "- 0, 1, .. : The specified gpu index \n"
             ),
         ),
+        'mediaprofile': MediaProfileSetting(
+            attr_name='mediaprofile',
+            nice_name='Media Profile',
+            value='',
+            tool_tip=('Adds the selected Media Profile to the command line')
+        ),
     }
 
     ndisplay_monitor_ui = None
@@ -611,6 +617,7 @@ class DevicenDisplay(DeviceUnreal):
             DevicenDisplay.csettings['udpmessaging_unicast_endpoint'],
             DevicenDisplay.csettings['udpmessaging_extra_static_endpoints'],
             DevicenDisplay.csettings['livelink_preset'],
+            DevicenDisplay.csettings['mediaprofile'],
             DevicenDisplay.csettings['graphics_adapter'],
             CONFIG.ENGINE_DIR,
             CONFIG.SOURCE_CONTROL_WORKSPACE,
@@ -809,22 +816,12 @@ class DevicenDisplay(DeviceUnreal):
         if DevicenDisplay.csettings['disable_all_screen_messages'].get_value() and 'DisableAllScreenMessages' not in exec_cmds:
             exec_cmds.append('DisableAllScreenMessages')
         
-        # LiveLink presets can be applied as ExecCmds
-        # e.g. of command:
-        #   "LiveLink.Preset.Apply Preset=/Game/Folder/MyLiveLinkPreset.MyLiveLinkPreset"
-
-        livelink_preset_gamepath = DeviceUnreal.csettings["livelink_preset"].get_value(self.name)
-
+        # LiveLink preset
+        livelink_preset_gamepath = DevicenDisplay.csettings["livelink_preset"].get_value(self.name)
         if livelink_preset_gamepath:
+            exec_cmds.append(self.exec_command_for_livelink_preset(livelink_preset_gamepath))
 
-            ext = '.uasset'
-
-            if livelink_preset_gamepath.endswith(ext):
-                livelink_preset_gamepath = livelink_preset_gamepath[:len(livelink_preset_gamepath)-len(ext)]
-            
-            preset_name = os.path.basename(os.path.normpath(livelink_preset_gamepath))
-            exec_cmds.append(f"LiveLink.Preset.Apply Preset={livelink_preset_gamepath}.{preset_name}")
-
+        # Format exec cmds
         exec_cmds = [cmd for cmd in exec_cmds if len(cmd.strip())]
 
         if len(exec_cmds):
@@ -903,6 +900,11 @@ class DevicenDisplay(DeviceUnreal):
         # Always tell Chaos to be deterministic
         dp_cvars.append('p.Chaos.Solver.Deterministic=1')
 
+        # mediaprofile
+        mediaprofile_gamepath = DevicenDisplay.csettings["mediaprofile"].get_value(self.name)
+        if mediaprofile_gamepath:
+            dp_cvars.append(self.dpcvar_for_mediaprofile(mediaprofile_gamepath))
+
         # Add user set dp cvars, overriding any of the forced ones.
         user_dp_cvars = self.csettings['ndisplay_dp_cvars'].get_value(self.name)
         user_dp_cvars = [cvar.strip() for cvar in user_dp_cvars if len(cvar.strip()) and len(cvar.split('=')) == 2]
diff --git a/Engine/Plugins/VirtualProduction/Switchboard/Source/Switchboard/switchboard/devices/unreal/plugin_unreal.py b/Engine/Plugins/VirtualProduction/Switchboard/Source/Switchboard/switchboard/devices/unreal/plugin_unreal.py
--- a/Engine/Plugins/VirtualProduction/Switchboard/Source/Switchboard/switchboard/devices/unreal/plugin_unreal.py
+++ b/Engine/Plugins/VirtualProduction/Switchboard/Source/Switchboard/switchboard/devices/unreal/plugin_unreal.py
@@ -346,6 +346,21 @@ class LiveLinkPresetSetting(Setting):
 
         self._on_widget_value_changed(value_str, override_device_name=override_device_name)
 
+
+    def asset_is_relevant(self, itemData) -> bool:
+        ''' Convenience function to filter the relevant assets for this setting '''
+
+        # Only return assets of the correct live link preset class
+        if itemData['classname'] not in DeviceUnreal.LIVELINKPRESET_CLASS_NAMES:
+            return False
+            
+        # Only return live link presets that have a valid gamepath. 
+        # We wouldn't be able to build the command line otherwise
+        if itemData['gamepath'] == '':
+            return False
+
+        return True
+
     def _update_combo_items(self, combo:QtWidgets.QComboBox, override_device_name):
         '''
         Populate the combobox itself with known list of live link presets available.
@@ -364,23 +379,9 @@ class LiveLinkPresetSetting(Setting):
 
         itemDatas = DeviceUnreal.csettings['asset_itemDatas'].get_value()
 
-        def asset_is_livelink_preset(itemData) -> bool:
-            ''' Convenience function to filter the live link preset assets '''
-
-            # Only return assets of the correct live link preset class
-            if itemData['classname'] not in DeviceUnreal.LIVELINKPRESET_CLASS_NAMES:
-                return False
-                
-            # Only return live link presets that have a valid gamepath. 
-            # We wouldn't be able to build the command line otherwise
-            if itemData['gamepath'] == '':
-                return False
-
-            return True
-
         # generate the combo box items
         try:
-            for itemData in [itemData for itemData in itemDatas if asset_is_livelink_preset(itemData)]:
+            for itemData in [itemData for itemData in itemDatas if self.asset_is_relevant(itemData)]:
                 name = itemData['name']
 
                 # trim the expected .uasset extension
@@ -390,8 +391,8 @@ class LiveLinkPresetSetting(Setting):
 
                 combo.addItem(name, itemData)
 
-        except Exception:
-            LOGGER.error('Error recalling asset itemDatas')
+        except Exception as e:
+            LOGGER.error(f'Error recalling asset itemDatas: {e}')
 
         # set the current index to the live link preset that was already selected
         for item_idx in range(combo.count()):
@@ -420,6 +421,24 @@ class LiveLinkPresetSetting(Setting):
                     combo.setCurrentIndex(item_idx)
                     break
 
+class MediaProfileSetting(LiveLinkPresetSetting):
+    ''' Container of the MediaProfile setting.
+    '''
+
+    #@override
+    def asset_is_relevant(self, itemData) -> bool:
+        # Only return assets of the correct live link preset class
+        if itemData['classname'] not in DeviceUnreal.MEDIAPROFILE_CLASS_NAMES:
+            return False
+            
+        # Only return live link presets that have a valid gamepath. 
+        # We wouldn't be able to build the command line otherwise
+        if itemData['gamepath'] == '':
+            return False
+
+        return True
+
+
 class DeviceUnreal(Device):
 
     NDISPLAY_CLASS_NAMES =  (
@@ -431,7 +450,12 @@ class DeviceUnreal(Device):
         'LiveLinkPreset',
         '/Script/LiveLink.LiveLinkPreset',
     )
-    
+
+    MEDIAPROFILE_CLASS_NAMES = (
+        'MediaProfile',
+        '/Script/MediaFrameworkUtilities.MediaProfile',
+    )
+
     csettings = {
         'buffer_size': IntSetting(
             attr_name="buffer_size",
@@ -613,6 +637,12 @@ class DeviceUnreal(Device):
             tool_tip="Remember the last analyzed list of project and plugin assets",
             show_ui=False,
         ),
+        'mediaprofile': MediaProfileSetting(
+            attr_name='mediaprofile',
+            nice_name='Media Profile',
+            value='',
+            tool_tip=('Adds the selected Media Profile to the command line')
+        ),
     }
 
     unreal_started_signal = QtCore.Signal()
@@ -951,6 +981,7 @@ class DeviceUnreal(Device):
             DeviceUnreal.csettings['udpmessaging_extra_static_endpoints'],
             DeviceUnreal.csettings['unrealgamesync_lib_dir'],
             DeviceUnreal.csettings['livelink_preset'],
+            DeviceUnreal.csettings['mediaprofile'],
             CONFIG.ENGINE_DIR,
             CONFIG.SOURCE_CONTROL_WORKSPACE,
             CONFIG.UPROJECT_PATH,
@@ -997,6 +1028,43 @@ class DeviceUnreal(Device):
 
         return valid
 
+    def conform_asset_gamepath(self, gamepath : str, ext : str = '.uasset') -> str:
+        ''' Conforms the given gamepath, which can include file extension, to the //Game/../MyAsset.MyAsset convention
+        e.g.
+            /Game/Folder/MyAssetName.uasset -> /Game/Folder/MyAssetName.MyAssetName
+        '''
+
+        if gamepath.endswith(ext):
+            gamepath = gamepath[:len(gamepath)-len(ext)]
+        
+        name = os.path.basename(os.path.normpath(gamepath))
+
+        return f"{gamepath}.{name}"
+
+    def exec_command_for_livelink_preset(self, livelink_preset_gamepath : str) -> str:
+        ''' Returns the exec command string to enable the given livelink preset gamepath
+        LiveLink presets can be applied as ExecCmds
+        e.g. of command:
+            "LiveLink.Preset.Apply Preset=/Game/Folder/MyLiveLinkPreset.MyLiveLinkPreset"
+        '''
+
+        if not livelink_preset_gamepath:
+            raise ValueError
+
+        return f"LiveLink.Preset.Apply Preset={self.conform_asset_gamepath(livelink_preset_gamepath)}"
+
+    def dpcvar_for_mediaprofile(self, mediaprofile_gamepath : str) -> str:
+        ''' Returns the dpcvar assignment string to enable the given mediaprofile gamepath
+        Media profiles can be applied as early cvars
+        e.g.:
+          "MediaUtils.StartupProfile=/Game/Media/MyMediaProfile.MyMediaProfile"
+        '''
+
+        if not mediaprofile_gamepath:
+            raise ValueError
+
+        return f"MediaUtils.StartupProfile={self.conform_asset_gamepath(mediaprofile_gamepath)}"
+
     @classmethod
     def plugin_header_widget_config(cls):
         """
@@ -1652,21 +1720,10 @@ class DeviceUnreal(Device):
 
         exec_cmds = DeviceUnreal.csettings["exec_cmds"].get_value(self.name).copy()
 
-        # LiveLink presets can be applied as ExecCmds
-        # e.g. of command:
-        #   "LiveLink.Preset.Apply Preset=/Game/Folder/MyLiveLinkPreset.MyLiveLinkPreset"
-
+        # LiveLink preset
         livelink_preset_gamepath = DeviceUnreal.csettings["livelink_preset"].get_value(self.name)
-
         if livelink_preset_gamepath:
-
-            ext = '.uasset'
-
-            if livelink_preset_gamepath.endswith(ext):
-                livelink_preset_gamepath = livelink_preset_gamepath[:len(livelink_preset_gamepath)-len(ext)]
-            
-            preset_name = os.path.basename(os.path.normpath(livelink_preset_gamepath))
-            exec_cmds.append(f"LiveLink.Preset.Apply Preset={livelink_preset_gamepath}.{preset_name}")
+            exec_cmds.append(self.exec_command_for_livelink_preset(livelink_preset_gamepath))
 
         # Exec Commands
 
@@ -1710,6 +1767,11 @@ class DeviceUnreal(Device):
         slate_allow_throttling = DeviceUnreal.csettings["slate_allow_throttling"].get_value()
         dp_cvars.append(f'Slate.bAllowThrottling={int(slate_allow_throttling)}')
 
+        # mediaprofile
+        mediaprofile_gamepath = DeviceUnreal.csettings["mediaprofile"].get_value(self.name)
+        if mediaprofile_gamepath:
+            dp_cvars.append(self.dpcvar_for_mediaprofile(mediaprofile_gamepath))
+
         # Add user set dp cvars, overriding any of the forced ones.
         user_dp_cvars = DeviceUnreal.csettings["dp_cvars"].get_value(self.name)
         user_dp_cvars = [cvar.strip() for cvar in user_dp_cvars if len(cvar.strip().split('=')) == 2]
@@ -2508,9 +2570,12 @@ class DeviceUnreal(Device):
         # Looks much better without the window frame.
         progressDiag.setWindowFlag(QtCore.Qt.FramelessWindowHint)
 
-        INTERESTING_CLASS_NAMES = DeviceUnreal.NDISPLAY_CLASS_NAMES + DeviceUnreal.LIVELINKPRESET_CLASS_NAMES
-
-        def validateInterestingAsset(asset):
+        INTERESTING_CLASS_NAMES = \
+            DeviceUnreal.NDISPLAY_CLASS_NAMES \
+            + DeviceUnreal.LIVELINKPRESET_CLASS_NAMES \
+            + DeviceUnreal.MEDIAPROFILE_CLASS_NAMES
+        
+        def validateInterestingAsset(asset):                                                                                                
             ''' Returns the asset if it is an interesting asset '''
 
             with open(asset['path'], 'rb') as file:
