Switchboard: Python sync script support for UGS AdditionalPathsToSync.

Other changes:
- Hides UGS/PCB support and behaviors behind `config.ENABLE_UGS_SUPPORT`
- Implements UGS revcx optimization for code change determination, which can save up to ~10 seconds during sync operations and reduce Perforce server load
- Fixes UGS subprocesses popping up dotnet console windows
- Removes unused import causing switchboard_utils.py to depend on PySide2
- Some static analysis fixes and PEP-8 cleanup

diff --git a/Engine/Plugins/VirtualProduction/Switchboard/Source/Switchboard/sbl_helper.py b/Engine/Plugins/VirtualProduction/Switchboard/Source/Switchboard/sbl_helper.py
--- a/Engine/Plugins/VirtualProduction/Switchboard/Source/Switchboard/sbl_helper.py
+++ b/Engine/Plugins/VirtualProduction/Switchboard/Source/Switchboard/sbl_helper.py
@@ -11,14 +11,13 @@ import logging
 import marshal
 import os
 import pathlib
-from platform import platform
 from queue import Empty, Queue
 import subprocess
 import sys
 import threading
 import tempfile
 import traceback
-from typing import Callable, Dict, List, Optional, Sequence, Tuple, Union
+from typing import Callable, Dict, List, Optional, Sequence, Union
 from switchboard import ugs_utils as UGS
 
 
@@ -117,6 +116,23 @@ def p4_changes(
     return p4_get_records('changes', opts, user=user, client=client)
 
 
+def p4_fstat(
+    pathspecs: List[str],
+    *,
+    limit: int = 0,
+    user: Optional[str] = None,
+    client: Optional[str] = None,
+) -> List[dict]:
+    opts: List[str] = []
+
+    if limit > 0:
+        opts.append(f'-m{limit}')
+
+    opts.extend(pathspecs)
+
+    return p4_get_records('fstat', opts, user=user, client=client)
+
+
 def p4_have(
     pathspecs: List[str],
     *,
@@ -191,6 +207,43 @@ def p4_print(
     return results
 
 
+def p4_latest_code_change(
+    paths: List[str],
+    *,
+    range: Optional[str] = None,
+    exts: Optional[List[str]] = None,
+    user: Optional[str] = None,
+    client: Optional[str] = None,
+) -> Optional[int]:
+    '''
+    Code CL determination compatible with UGS/precompiled binaries.
+    See `WorkspaceUpdate.ExecuteAsync` in `WorkspaceUpdate.cs`
+    '''
+
+    ALL_CODE_EXTS = [
+        '.c', '.cc', '.cpp', '.inl', '.m', '.mm', '.rc', '.cs', '.csproj',
+        '.h', '.hpp', '.usf', '.ush', '.uproject', '.uplugin', '.sln']
+
+    if exts is None:
+        exts = ALL_CODE_EXTS
+
+    rangespec = f'@{range}' if range is not None else ''
+
+    code_paths: List[str] = []
+    for path in paths:
+        if path.endswith('...'):
+            code_paths.extend([f'{path}/*{ext}{rangespec}'
+                                for ext in exts])
+
+    code_changes = p4_changes(code_paths, limit=1,
+                              user=user, client=client)
+
+    if len(code_changes) > 0:
+        return max(int(x[b'change']) for x in code_changes)
+    else:
+        return None
+
+
 class SyncFile:
     ''' Represents a single file pending sync. '''
     def __init__(
@@ -427,6 +480,7 @@ def p4_sync(
     write_workload_file: Optional[io.TextIOWrapper] = None
 ) -> Optional[int]:
     # First, gather complete list of files/revisions that need to be synced.
+    logging.info('Determining files to be synced')
     total_work = p4_get_preview_sync_files(input_specs, user=user,
                                            client=client)
 
@@ -649,6 +703,7 @@ class SbListenerHelper:
 
         self.uproj_path: Optional[pathlib.Path] = None
         self.engine_dir: Optional[pathlib.Path] = None
+        self.additional_paths_to_sync: List[str] = []
         self.sync_engine_cl: Optional[int] = None
         self.sync_project_cl: Optional[int] = None
         self.clobber_engine: bool = False
@@ -674,11 +729,6 @@ class SbListenerHelper:
         # Sync action
         sync_parser = subparsers.add_parser('sync')
 
-        sync_parser.add_argument(
-            'pathspecs', nargs='*', metavar='pathSpec[revSpec]',
-            help='File/path specifier, supporting * or ... wildcards '
-                 '[and optional # or @ revision specifier suffix]')
-
         sync_parser.add_argument(
             '--project', type=pathlib.Path, metavar='UPROJECT',
             help='Path to .uproject file')
@@ -697,13 +747,16 @@ class SbListenerHelper:
             help='Run GenerateProjectFiles after syncing')
         sync_parser.add_argument(
             '--use-ugs', action='store_true',
-            help="Specifies that UnrealGameSync should be used to sync the project and engine together.")
+            help='Specifies that UnrealGameSync should be used to sync the '
+                 'project and engine together.')
         sync_parser.add_argument(
             '--use-pcbs', action='store_true',
-            help="Specifies that precompiled binaries should be sync'd along with the project (requires syncing via UnrealGameSync).")
+            help='Specifies that precompiled binaries should be synced along '
+                 'with the project (requires syncing via UnrealGameSync).')
         sync_parser.add_argument(
             '--ugs-lib-dir', type=pathlib.Path,
-            help='Directory path specifying where to find the UnrealGameSync library (ugs.dll).')
+            help='Directory path specifying where to find the UnrealGameSync '
+                 'library (ugs.dll).')
         sync_parser.add_argument(
             '--clobber-engine', action='store_true',
             help='Override noclobber for engine files')
@@ -740,19 +793,29 @@ class SbListenerHelper:
         assert False
 
     def run_sync(self, options: argparse.Namespace) -> int:
-        pathspecs: List[str] = [*options.pathspecs]
+        engine_paths: List[str] = []
+        project_paths: List[str] = []
+        pathspecs: List[str] = []
 
         self.check_sync_options(options)
+        self.additional_paths_to_sync = self.get_additional_paths_to_sync()
 
         if self.sync_engine_cl:
-            pathspecs.append((f'{self.engine_dir.parent / "*"}'
-                              f'@{self.sync_engine_cl}'))
-            pathspecs.append((f'{self.engine_dir / "..."}'
-                              f'@{self.sync_engine_cl}'))
+            engine_paths.append(f'{self.engine_dir.parent / "*"}')
+            engine_paths.append(f'{self.engine_dir / "..."}')
+
+            pathspecs.extend(
+                [f'{path}@{self.sync_engine_cl}' for path in engine_paths])
 
         if self.sync_project_cl:
-            pathspecs.append((f'{self.uproj_path.parent / "..."}'
-                              f'@{self.sync_project_cl}'))
+            project_paths.append(f'{self.uproj_path.parent / "..."}')
+
+            for path in self.additional_paths_to_sync:
+                path = path.lstrip('/')
+                project_paths.append(f'{self.engine_dir.parent / path}')
+
+            pathspecs.extend(
+                [f'{path}@{self.sync_project_cl}' for path in project_paths])
 
         if len(pathspecs) < 1:
             self.parser.error('Nothing specified to sync')
@@ -761,11 +824,14 @@ class SbListenerHelper:
             for pathspec in pathspecs:
                 logging.info(f' - {pathspec}')
 
-        buildver_info = self.get_buildver_info()
-        if self.sync_engine_cl:
-            # Remove Build.version if present, so we can write our own.
-            if buildver_info.local_path and buildver_info.have_rev:
-                pathspecs.append(f"{buildver_info.local_path}#0")
+        if not self.use_ugs:
+            logging.info('Determining CompatibleChangelist')
+            buildver_info = self.get_buildver_info(engine_paths, project_paths)
+
+            if self.sync_engine_cl:
+                # Remove Build.version if present, so we can write our own.
+                if buildver_info.local_path and buildver_info.have_rev:
+                    pathspecs.append(f"{buildver_info.local_path}#0")
 
         def sync_filter(workload: SyncWorkload):
             # Don't add Build.version to the workspace, but do allow removal.
@@ -822,7 +888,7 @@ class SbListenerHelper:
             return sync_result  # Nothing more to do
 
         # Write updated Build.version.
-        if self.sync_engine_cl:
+        if not self.use_ugs:
             if buildver_info.local_path and buildver_info.updated_text:
                 logging.info('Updating Build.version')
                 try:
@@ -831,9 +897,6 @@ class SbListenerHelper:
                 except Exception as exc:
                     logging.error('Exception writing to Build.version',
                                   exc_info=exc)
-            else:
-                logging.warning('Unable to update Build.version at path '
-                                f'{buildver_info.local_path}')
 
         # Generate project files.
         if self.generate_after_sync:
@@ -856,13 +919,15 @@ class SbListenerHelper:
         depot_text: Optional[str] = None
         updated_text: Optional[str] = None
 
-    def get_buildver_info(self) -> BuildVerInfo:
+    def get_buildver_info(
+        self,
+        engine_paths: List[str],
+        project_paths: List[str],
+    ) -> BuildVerInfo:
         '''
         Read `Build.version` from depot and return its info, depot contents,
         and an updated version of its contents which reflects the local build.
         '''
-        logging.debug('get_buildver_info(): start')
-
         ret_info = SbListenerHelper.BuildVerInfo()
 
         if not self.engine_dir:
@@ -876,8 +941,9 @@ class SbListenerHelper:
 
         ret_info.local_path = buildver_local_path
 
-        ret_info.have_rev = p4_have([buildver_local_path],
-                                    user=self.p4user, client=self.p4client)[0]
+        p4opts = {'user': self.p4user, 'client': self.p4client}
+
+        ret_info.have_rev = p4_have([buildver_local_path], **p4opts)[0]
 
         if not self.sync_engine_cl:
             return ret_info
@@ -885,67 +951,170 @@ class SbListenerHelper:
         prints = p4_print(
             [f'{path}@{self.sync_engine_cl}'
              for path in (epicint_local_path, buildver_local_path)],
-            user=self.p4user, client=self.p4client)
+            **p4opts)
         epicint_print_result = prints[0]
         buildver_print_result = prints[1]
 
         ret_info.depot_text = buildver_print_result.text  # May be None
+        buildver_meta = buildver_print_result.stat
 
-        if buildver_print_result.is_valid:
-            buildver_meta = buildver_print_result.stat
-            ret_info.depot_path = buildver_meta[b'depotFile'].decode()
-
-            # Code CL determination compatible with precompiled binaries.
-            code_exts = ['.cs', '.h', '.cpp', '.usf', '.ush', '.uproject',
-                         '.uplugin']
-            engine_code_paths = [
-                f'{self.engine_dir}/.../*{ext}@<={self.sync_engine_cl}'
-                for ext in code_exts]
-            engine_code_changes = p4_changes(
-                engine_code_paths, limit=1,
-                user=self.p4user, client=self.p4client)
-            if len(engine_code_changes) > 0:
-                code_cl = max(int(x[b'change']) for x in engine_code_changes)
-            else:
-                logging.warning('Unable to determine last code changelist')
-                code_cl = self.sync_engine_cl
+        if not buildver_print_result.is_valid:
+            logging.error('Unable to p4_print() Build.version '
+                          f'({buildver_meta})')
 
-            # Generate and return the updated Build.version contents.
-            try:
-                build_ver: Dict[str, Union[int, str]] = json.loads(
-                    ret_info.depot_text)
-            except json.JSONDecodeError:
-                logging.error('Unable to parse depot Build.version JSON',
-                              exc_info=sys.exc_info())
-                return ret_info
-
-            branch_name = ret_info.depot_path.replace(
-                '/Engine/Build/Build.version', '')
-            updates = {
-                'Changelist': self.sync_engine_cl,
-                'IsLicenseeVersion': 0 if epicint_print_result.is_valid else 1,
-                'BranchName': branch_name.replace('/', '+'),
-            }
-
-            # Don't overwrite the compatible changelist if we're in a hotfix
-            # release.
-            no_depot_compat = build_ver['CompatibleChangelist'] == 0
-            licensee_changed = (build_ver['IsLicenseeVersion']
-                                != updates['IsLicenseeVersion'])
-            if no_depot_compat or licensee_changed:
-                updates['CompatibleChangelist'] = code_cl
-
-            logging.debug(f"get_buildver_info(): updates = {updates}")
-
-            build_ver.update(updates)
-            ret_info.updated_text = json.dumps(build_ver, indent=4)
+            return ret_info
+
+        ret_info.depot_path = buildver_meta[b'depotFile'].decode()
+
+        # Determine latest code change within each group (by CL) of paths.
+        # First try to narrow using revcx-friendly check (1000 CLs, just *.cpp)
+        lower_bound_cl: Optional[int] = None
+        eng_upper_bound_cl = self.sync_engine_cl
+        proj_upper_bound_cl = self.sync_project_cl
+
+        def quick_narrow_lower_cl(paths: List[str], upper_cl: int,
+                                  in_lower_cl: Optional[int]) -> Optional[int]:
+            quick_lower = upper_cl - 1_000
+            if in_lower_cl and in_lower_cl > quick_lower:
+                quick_lower = in_lower_cl
+            if upper_cl <= quick_lower:
+                return upper_cl
+            if quick_lower > 0:
+                quick_range = f'{quick_lower},{upper_cl}'
+                check_lower = p4_latest_code_change(
+                    paths, range=quick_range, exts=['.cpp'], **p4opts)
+
+                if check_lower:
+                    if in_lower_cl:
+                        return max(in_lower_cl, check_lower)
+                    else:
+                        return check_lower
 
+            return in_lower_cl
+
+        lower_bound_cl = quick_narrow_lower_cl(engine_paths,
+                                               eng_upper_bound_cl,
+                                               lower_bound_cl)
+
+        if proj_upper_bound_cl:
+            lower_bound_cl = quick_narrow_lower_cl(project_paths,
+                                                   proj_upper_bound_cl,
+                                                   lower_bound_cl)
+
+        # Comprehensive check (all exts), but hopefully with narrowed ranges
+        eng_range = f'<={eng_upper_bound_cl}'
+        proj_range = f'<={proj_upper_bound_cl}'
+
+        if lower_bound_cl:
+            eng_range = f'{lower_bound_cl},{eng_upper_bound_cl}'
+            proj_range = f'{lower_bound_cl},{proj_upper_bound_cl}'
+
+        eng_code_cl = p4_latest_code_change(engine_paths, range=eng_range,
+                                            **p4opts)
+        proj_code_cl = p4_latest_code_change(project_paths, range=proj_range,
+                                             **p4opts)
+
+        # These are the values we'll actually write to the JSON
+        ver_cl = max(filter(None, [self.sync_engine_cl, self.sync_project_cl]))
+        ver_compatible_cl = max(filter(None, [eng_code_cl, proj_code_cl]))
+
+        # Generate and return the updated Build.version contents.
+        try:
+            build_ver: Dict[str, Union[int, str]] = json.loads(
+                ret_info.depot_text)
+        except json.JSONDecodeError:
+            logging.error('Unable to parse depot Build.version JSON',
+                          exc_info=sys.exc_info())
+            return ret_info
+
+        branch_name = ret_info.depot_path.replace(
+            '/Engine/Build/Build.version', '')
+        updates = {
+            'Changelist': ver_cl,
+            'IsLicenseeVersion': 0 if epicint_print_result.is_valid else 1,
+            'BranchName': branch_name.replace('/', '+'),
+        }
+
+        # Don't overwrite the compatible changelist if we're in a hotfix
+        # release.
+        no_depot_compat = build_ver['CompatibleChangelist'] == 0
+        licensee_changed = (build_ver['IsLicenseeVersion']
+                            != updates['IsLicenseeVersion'])
+        if no_depot_compat or licensee_changed:
+            updates['CompatibleChangelist'] = ver_compatible_cl
+
+        updated_build_ver = build_ver.copy()
+        updated_build_ver.update(updates)
+        if updated_build_ver == build_ver:
+            logging.info("get_buildver_info(): contents haven't changed")
         else:
-            logging.error('Unable to p4_print() Build.version '
-                          f'({buildver_meta})')
+            logging.info(f"get_buildver_info(): {updates}")
+            ret_info.updated_text = json.dumps(updated_build_ver, indent=4)
 
         return ret_info
 
+    def get_additional_paths_to_sync(self) -> List[str]:
+        # We append depot wildcards to host paths and end up with mixed syntax
+        # here, but fstat separates them into normalized depot/host paths.
+        mixed_config_paths = UGS.get_depot_config_paths(self.engine_dir,
+                                                        self.uproj_path.parent)
+
+        p4opts = {'user': self.p4user, 'client': self.p4client}
+
+        config_fstats = p4_fstat(mixed_config_paths, **p4opts)
+
+        config_contents_map: Dict[str, Optional[str]] = {}
+
+        # First pass: get contents of locally modified files, or depot paths
+        # of files to `p4 print` in the second pass otherwise
+        for record in config_fstats:
+            action: Optional[bytes] = record.get(b'action')
+            head_action: Optional[bytes] = record.get(b'headAction')
+            if record[b'code'] == b'stat':
+                if (b'clientFile' in record) and (b'action' in record):
+                    # If the client file exists and is modified, use local ver
+                    client_path = pathlib.Path(record[b'clientFile'].decode())
+                    if client_path.exists and action == b'edit':
+                        local_contents = client_path.read_text()
+                        config_contents_map[str(client_path)] = local_contents
+                        continue
+                    elif action == b'delete':
+                        continue
+
+                if head_action == b'delete':
+                    continue
+
+                if b'depotFile' in record:
+                    depot_file = record[b'depotFile'].decode()
+                    config_contents_map[depot_file] = None
+                    continue
+
+            logging.warning(f'Unhandled fstat record: {record}')
+
+        # Second pass: p4 prints of depot configs @ head
+        depot_print_paths = [k for k, v in config_contents_map.items()
+                             if v is None]  # configs we didn't read locally
+
+        config_prints = p4_print(depot_print_paths, **p4opts)
+
+        for depot_path, record in zip(depot_print_paths, config_prints):
+            if record.is_valid:
+                config_contents_map[depot_path] = record.text
+            else:
+                logging.warning(f'Invalid P4PrintResult for {depot_path}')
+
+        parser = UGS.IniParser()
+        for path, contents in config_contents_map.items():
+            if contents is None:
+                logging.warning(f'No contents for {path}')
+                continue
+
+            parser.read_string(contents, path)
+
+        paths = parser.try_get('Perforce', 'AdditionalPathsToSync')
+        logging.debug(f'AdditionalPathsToSync: {paths}')
+        return paths or []
+
     def on_sync_progress(
         self,
         completed_bytes: int,
@@ -996,14 +1165,17 @@ class SbListenerHelper:
         self.sync_engine_cl = options.engine_cl
         self.sync_project_cl = options.project_cl
         self.use_ugs = options.use_ugs
-        self.generate_after_sync = options.generate and not self.use_ugs # UGS will generate the project files by default (no need to duplicate work)
+        # UGS will generate the project files by default (don't duplicate work)
+        self.generate_after_sync = options.generate and not self.use_ugs
         self.sync_pcbs = options.use_pcbs
 
         if self.sync_pcbs and not self.use_ugs:
             self.parser.error('`--use-pcbs` requires `--use-ugs`.')
-        
+
         if self.use_ugs and self.sync_engine_cl:
-            self.parser.error('`--use-ugs` was specified along with `--engine-cl`; these arguments are mutally exclusive and should not be used together.')
+            self.parser.error('`--use-ugs` was specified along with '
+                              '`--engine-cl`; these arguments are mutally '
+                              'exclusive and should not be used together.')
 
         if options.project:
             try:
diff --git a/Engine/Plugins/VirtualProduction/Switchboard/Source/Switchboard/switchboard/config.py b/Engine/Plugins/VirtualProduction/Switchboard/Source/Switchboard/switchboard/config.py
--- a/Engine/Plugins/VirtualProduction/Switchboard/Source/Switchboard/switchboard/config.py
+++ b/Engine/Plugins/VirtualProduction/Switchboard/Source/Switchboard/switchboard/config.py
@@ -32,6 +32,9 @@ USER_SETTINGS_BACKUP_FILE_PATH = ROOT_CONFIGS_PATH.joinpath(USER_SETTINGS_BACKUP
 
 DEFAULT_MAP_TEXT = '-- Default Map --'
 
+ENABLE_UGS_SUPPORT = False
+
+
 def migrate_comma_separated_string_to_list(value) -> typing.List[str]:
     if isinstance(value, str):
         return value.split(",")
@@ -40,6 +43,7 @@ def migrate_comma_separated_string_to_list(value) -> typing.List[str]:
         return value
     raise NotImplementedError("Migration not handled")
 
+
 class Setting(QtCore.QObject):
     '''
     A type-agnostic value container for a configuration setting.
@@ -1869,11 +1873,15 @@ class ConfigPathValidator(QtGui.QValidator):
 
         return QtGui.QValidator.Acceptable
 
+
 class EngineSyncMethod(Enum):
     Use_Existing = "Use Existing (do not sync/build)"
     Build_Engine = "Build Engine"
-    Sync_PCBs = "Sync Precompiled Binaries (requires UnrealGameSync)"
-    Sync_From_UGS = "Sync engine and project together using UnrealGameSync"
+
+    if ENABLE_UGS_SUPPORT:
+        Sync_PCBs = "Sync Precompiled Binaries (requires UnrealGameSync)"
+        Sync_From_UGS = "Sync engine and project together using UnrealGameSync"
+
 
 class Config(object):
 
@@ -2100,7 +2108,7 @@ class Config(object):
             ),
         }
 
-        # Done here outside of the setting's initializor so that values different from the default are flagged in the UI (with a 'reset' option) 
+        # Done here outside of the setting's initializer so that values different from the default are flagged in the UI (with a 'reset' option) 
         if 'engine_sync_method' in data:
             self.basic_project_settings["engine_sync_method"].update_value(data.get('engine_sync_method'))
         # To support backwards compatibility, we check the old 'build_engine' option (from before we had a 'PCB' option)
diff --git a/Engine/Plugins/VirtualProduction/Switchboard/Source/Switchboard/switchboard/devices/unreal/plugin_unreal.py b/Engine/Plugins/VirtualProduction/Switchboard/Source/Switchboard/switchboard/devices/unreal/plugin_unreal.py
--- a/Engine/Plugins/VirtualProduction/Switchboard/Source/Switchboard/switchboard/devices/unreal/plugin_unreal.py
+++ b/Engine/Plugins/VirtualProduction/Switchboard/Source/Switchboard/switchboard/devices/unreal/plugin_unreal.py
@@ -24,8 +24,8 @@ from switchboard import config_osc as osc
 from switchboard import message_protocol, switchboard_application
 from switchboard import switchboard_dialog as sb_dialog
 from switchboard import switchboard_utils as sb_utils
-from switchboard.config import (CONFIG, DEFAULT_MAP_TEXT, SETTINGS, BoolSetting, DirectoryPathSetting, EngineSyncMethod,
-                                FilePathSetting, IntSetting, MultiOptionSetting, OptionSetting, Setting,
+from switchboard.config import (CONFIG, DEFAULT_MAP_TEXT, ENABLE_UGS_SUPPORT, SETTINGS, BoolSetting, DirectoryPathSetting,
+                                EngineSyncMethod, FilePathSetting, IntSetting, MultiOptionSetting, OptionSetting, Setting,
                                 StringListSetting, StringSetting, migrate_comma_separated_string_to_list)
 from switchboard.devices.device_base import Device, DeviceStatus, PluginHeaderWidgets
 from switchboard.devices.device_widget_base import DeviceAutoJoinMUServerUI, DeviceWidget
@@ -598,15 +598,6 @@ class DeviceUnreal(Device):
                 'between network messages before considering the connection '
                 'to Switchboard lost and closing it with a timeout error.')
         ),
-        'unrealgamesync_lib_dir': DirectoryPathSetting(
-            attr_name='unrealgamesync_lib_dir',
-            nice_name='UnrealGameSync Dll Dir',
-            value="",
-            tool_tip=(
-                "Expects an absolute path to the device directory containing the ugs.dll library file. \n"
-                "If left blank, switchboard will attempt to find UGS by searching the system PATH and default install locations. \n"
-                "On Windows, the default install location is '${LOCALAPPDATA}/UnrealGameSync/Latest/'.")
-        ),
         'slate_allow_throttling': BoolSetting(
             attr_name='slate_allow_throttling',
             nice_name='Allow Slate Throttling',
@@ -645,6 +636,17 @@ class DeviceUnreal(Device):
         ),
     }
 
+    if ENABLE_UGS_SUPPORT:
+        csettings['unrealgamesync_lib_dir'] = DirectoryPathSetting(
+            attr_name='unrealgamesync_lib_dir',
+            nice_name='UnrealGameSync Dll Dir',
+            value="",
+            tool_tip=(
+                "Expects an absolute path to the device directory containing the ugs.dll library file. \n"
+                "If left blank, switchboard will attempt to find UGS by searching the system PATH and default install locations. \n"
+                "On Windows, the default install location is '${LOCALAPPDATA}/UnrealGameSync/Latest/'.")
+        )
+
     unreal_started_signal = QtCore.Signal()
 
     mu_server = switchboard_application.get_multi_user_server_instance()
@@ -969,7 +971,7 @@ class DeviceUnreal(Device):
         return Device.plugin_settings() + list(DeviceUnreal.csettings.values())
 
     def setting_overrides(self):
-        return super().setting_overrides() + [
+        overrides = super().setting_overrides() + [
             Device.csettings['is_recording_device'],
             DeviceUnreal.csettings['command_line_arguments'],
             DeviceUnreal.csettings['exec_cmds'],
@@ -979,7 +981,6 @@ class DeviceUnreal(Device):
             DeviceUnreal.csettings['auto_decline_package_recovery'],
             DeviceUnreal.csettings['udpmessaging_unicast_endpoint'],
             DeviceUnreal.csettings['udpmessaging_extra_static_endpoints'],
-            DeviceUnreal.csettings['unrealgamesync_lib_dir'],
             DeviceUnreal.csettings['livelink_preset'],
             DeviceUnreal.csettings['mediaprofile'],
             CONFIG.ENGINE_DIR,
@@ -987,6 +988,11 @@ class DeviceUnreal(Device):
             CONFIG.UPROJECT_PATH,
         ]
 
+        if ENABLE_UGS_SUPPORT:
+            overrides.append(DeviceUnreal.csettings['unrealgamesync_lib_dir'])
+
+        return overrides
+
     def device_settings(self):
         return super().device_settings() + [
             self.setting_roles,
@@ -1113,7 +1119,7 @@ class DeviceUnreal(Device):
                 self._request_engine_changelist_number()
         else:
             self.widget.engine_changelist_label.hide()
-            self.widget.update_build_info(current_cl=self.engine_changelist, built_cl=self.built_engine_changelist)
+            self.widget.update_build_info(synched_cl=self.engine_changelist, built_cl=self.built_engine_changelist)
 
     def on_setting_exclude_from_build_changed(self, exclude_from_build):
         self.widget.update_exclude_from_build(exclude_from_build, not self.is_disconnected)
@@ -1306,8 +1312,9 @@ class DeviceUnreal(Device):
 
         sync_method = CONFIG.ENGINE_SYNC_METHOD.get_value()
         generate_proj_files = sync_method == EngineSyncMethod.Build_Engine.value
-        sync_precompiled_bins = sync_method == EngineSyncMethod.Sync_PCBs.value
-        sync_using_ugs = sync_precompiled_bins or sync_method == EngineSyncMethod.Sync_From_UGS.value
+        if ENABLE_UGS_SUPPORT:
+            sync_precompiled_bins = sync_method == EngineSyncMethod.Sync_PCBs.value
+            sync_using_ugs = sync_precompiled_bins or sync_method == EngineSyncMethod.Sync_From_UGS.value
 
         project_name = os.path.basename(os.path.dirname(project_path))
         LOGGER.info(
@@ -1349,18 +1356,19 @@ class DeviceUnreal(Device):
         if generate_proj_files:
             sync_args += ' --generate'
 
-        if sync_using_ugs:
-            sync_args += ' --use-ugs'
-            if self.unrealgamesync_lib_dir_setting:
-                sync_args += f' --ugs-lib-dir={self.unrealgamesync_lib_dir_setting}'
+        if ENABLE_UGS_SUPPORT:
+            if sync_using_ugs:
+                sync_args += ' --use-ugs'
+                if self.unrealgamesync_lib_dir_setting:
+                    sync_args += f' --ugs-lib-dir={self.unrealgamesync_lib_dir_setting}'
         
-        if sync_precompiled_bins:
-            sync_args += ' --use-pcbs'
+            if sync_precompiled_bins:
+                sync_args += ' --use-pcbs'
             
-            # If we're syncing 'Precompiled Binaries' one of those binaries may be the SwitchboardListener executable
-            # which means (on Windows atleast) we need to move the executable to make way for the new one
-            _, msg = message_protocol.create_free_listener_bin_message()
-            self.unreal_client.send_message(msg)
+                # If we're syncing 'Precompiled Binaries' one of those binaries may be the SwitchboardListener executable
+                # which means (on Windows atleast) we need to move the executable to make way for the new one
+                _, msg = message_protocol.create_free_listener_bin_message()
+                self.unreal_client.send_message(msg)
 
         puuid, msg = message_protocol.create_start_process_message(
             prog_path=sync_tool,
diff --git a/Engine/Plugins/VirtualProduction/Switchboard/Source/Switchboard/switchboard/switchboard_dialog.py b/Engine/Plugins/VirtualProduction/Switchboard/Source/Switchboard/switchboard/switchboard_dialog.py
--- a/Engine/Plugins/VirtualProduction/Switchboard/Source/Switchboard/switchboard/switchboard_dialog.py
+++ b/Engine/Plugins/VirtualProduction/Switchboard/Source/Switchboard/switchboard/switchboard_dialog.py
@@ -28,7 +28,7 @@ from switchboard import switchboard_application
 from switchboard import switchboard_utils
 from switchboard import switchboard_widgets as sb_widgets
 from switchboard.add_config_dialog import AddConfigDialog
-from switchboard.config import CONFIG, DEFAULT_MAP_TEXT, SETTINGS, EngineSyncMethod
+from switchboard.config import CONFIG, DEFAULT_MAP_TEXT, ENABLE_UGS_SUPPORT, SETTINGS, EngineSyncMethod
 from switchboard.device_list_widget import DeviceListWidget, DeviceWidgetHeader
 from switchboard.devices.device_base import DeviceStatus
 from switchboard.devices.device_manager import DeviceManager
@@ -1768,11 +1768,12 @@ class SwitchboardDialog(QtCore.QObject):
 
         changelists = None
         # If we're syncing 'Precompiled Binaries', then that implies that we should be using UGS:
-        if sync_method == EngineSyncMethod.Sync_PCBs.value or sync_method == EngineSyncMethod.Sync_From_UGS.value:
-            LOGGER.info("Using UnrealGameSync to refresh project changelists.")
-            changelists = ugs_utils.latest_chagelists(Path(CONFIG.UPROJECT_PATH.get_value()), client=CONFIG.SOURCE_CONTROL_WORKSPACE.get_value())
-            if not changelists:
-                LOGGER.error("UnrealGameSync failed to get the project's latest changelists. Falling back to using p4 commands directly.")
+        if ENABLE_UGS_SUPPORT:
+            if sync_method == EngineSyncMethod.Sync_PCBs.value or sync_method == EngineSyncMethod.Sync_From_UGS.value:
+                LOGGER.info("Using UnrealGameSync to refresh project changelists.")
+                changelists = ugs_utils.latest_chagelists(Path(CONFIG.UPROJECT_PATH.get_value()), client=CONFIG.SOURCE_CONTROL_WORKSPACE.get_value())
+                if not changelists:
+                    LOGGER.error("UnrealGameSync failed to get the project's latest changelists. Falling back to using p4 commands directly.")
 
         if not changelists:
             LOGGER.info("Refreshing p4 project changelists")
diff --git a/Engine/Plugins/VirtualProduction/Switchboard/Source/Switchboard/switchboard/switchboard_utils.py b/Engine/Plugins/VirtualProduction/Switchboard/Source/Switchboard/switchboard/switchboard_utils.py
--- a/Engine/Plugins/VirtualProduction/Switchboard/Source/Switchboard/switchboard/switchboard_utils.py
+++ b/Engine/Plugins/VirtualProduction/Switchboard/Source/Switchboard/switchboard/switchboard_utils.py
@@ -1,4 +1,5 @@
 # Copyright Epic Games, Inc. All Rights Reserved.
+
 import csv
 import datetime
 from enum import Enum
@@ -8,8 +9,6 @@ import subprocess
 import threading
 from typing import Optional
 
-from .switchboard_logging import LOGGER
-
 
 class PriorityModifier(Enum):
     ''' Corresponds to `int32 PriorityModifier` parameter to `FPlatformProcess::CreateProc`. '''
diff --git a/Engine/Plugins/VirtualProduction/Switchboard/Source/Switchboard/switchboard/ugs_utils.py b/Engine/Plugins/VirtualProduction/Switchboard/Source/Switchboard/switchboard/ugs_utils.py
--- a/Engine/Plugins/VirtualProduction/Switchboard/Source/Switchboard/switchboard/ugs_utils.py
+++ b/Engine/Plugins/VirtualProduction/Switchboard/Source/Switchboard/switchboard/ugs_utils.py
@@ -1,13 +1,16 @@
 # Copyright Epic Games, Inc. All Rights Reserved.
 
+import io
 import logging
-import subprocess
-import sys
-import shutil
-import re
 from os import path
 from pathlib import Path
-from typing import Optional, List, Dict, Callable
+import re
+import shutil
+import subprocess
+import sys
+from typing import Callable, Dict, List, Optional, TextIO, Union
+
+from . import switchboard_utils as sb_utils
 
 
 def find_bin(
@@ -21,50 +24,74 @@ def find_bin(
             ugs_dir = Path(env_ugs_path).parent
         elif sys.platform.startswith('win'):
             ugs_dir = path.join('${LOCALAPPDATA}', 'UnrealGameSync', 'Latest')
-        #elif: for other platforms we don't have a default install location (the dll is usually installed from perforce, which could go anywhere)
+        # elif:
+            # for other platforms we don't have a default install location (the
+            # dll is usually installed from perforce, which could go anywhere)
+
+    INSTALL_PATH_SUGGESTION = (
+        'Ensure that UGS is installed on the target device. If it is already '
+        'installed, then either make sure the `ugs` command is globally '
+        'available, or specify the directory explicitly.'
+    )
 
     if not ugs_dir:
-        logging.error('Could not determine where to find the UnrealGameSync library (ugs.dll). Ensure that UGS is installed on the target device.'
-                      'If it already is installed, then either make sure the `ugs` command is globally available, or specify the directory explicitly.')
+        logging.error('Could not determine where to find the UnrealGameSync '
+                      'library (ugs.dll). ' + INSTALL_PATH_SUGGESTION)
         return None
 
     ugs_bin_path = path.join(path.expandvars(ugs_dir), 'ugs.dll')
     if not path.exists(ugs_bin_path):
-        logging.error(f"Failed to find '{ugs_bin_path}'. Ensure that UGS is installed on the target device."
-                       'If it already is installed, then either make sure the `ugs` command is globally available, or specify the directory explicitly.')
+        logging.error(f"Failed to find '{ugs_bin_path}'. "
+                      + INSTALL_PATH_SUGGESTION)
         return None
 
-    return ugs_bin_path
+    return Path(ugs_bin_path)
+
 
 def setup_dependencies(
-    ue_engine_dir:Path
+    ue_engine_dir: Path
 ) -> int:
     logging.debug('Verifying UnrealGameSync dependencies.')
-    
+
     if sys.platform.startswith('win'):
         script_name = 'GetDotnetPath.bat'
-        dotnet_setup_script = path.join(ue_engine_dir, 'Build', 'BatchFiles', script_name)
+        dotnet_setup_script = path.join(ue_engine_dir, 'Build', 'BatchFiles',
+                                        script_name)
         dotnet_setup_args = [dotnet_setup_script]
     else:
-        platform_dirname = 'Mac' if sys.platform.startswith('darwin') else 'Linux'
-        platform_scripts_dir = path.join(ue_engine_dir, 'Build', 'BatchFiles', f'{platform_dirname}')
+        if sys.platform.startswith('darwin'):
+            platform_dirname = 'Mac'
+        else:
+            platform_dirname = 'Linux'
+
+        platform_scripts_dir = path.join(ue_engine_dir, 'Build', 'BatchFiles',
+                                         platform_dirname)
 
         script_name = 'SetupEnvironment.sh'
         dotnet_setup_script = path.join(platform_scripts_dir, script_name)
-        dotnet_setup_args = [dotnet_setup_script, '-dotnet', f'{platform_scripts_dir}']
+        dotnet_setup_args = [dotnet_setup_script, '-dotnet',
+                             platform_scripts_dir]
 
     try:
-        with subprocess.Popen(dotnet_setup_args, stdin=subprocess.DEVNULL, stdout=subprocess.PIPE) as proc:
+        with subprocess.Popen(dotnet_setup_args, stdin=subprocess.DEVNULL,
+                              stdout=subprocess.PIPE,
+                              startupinfo=sb_utils.get_hidden_sp_startupinfo()
+                              ) as proc:
             for line in proc.stdout:
                 logging.debug(f'{script_name}> {line.decode().rstrip()}')
     except Exception as exc:
-        logging.error(f'setup_ugs_dependencies(): exception running Popen: {dotnet_setup_args}', exc_info=exc)
+        logging.error('setup_ugs_dependencies(): exception running Popen: '
+                      f'{dotnet_setup_args}', exc_info=exc)
+        return -1
 
     if proc.returncode != 0:
-        logging.error('Unable to find a install of Dotnet SDK. Please make sure you have it installed and that `dotnet` is a globally available command.')
-    
+        logging.error('Unable to find a install of Dotnet SDK. Please ensure '
+                      'you have it installed and that `dotnet` is a globally '
+                      'available command.')
+
     return proc.returncode
 
+
 def _find_engine_dir(
     uproj_path: Path
 ) -> Optional[Path]:
@@ -78,118 +105,152 @@ def _find_engine_dir(
     # TODO: Consider looking for Engine dir as ancestor of this script?
     return None
 
+
 def _get_active_ugs_context(
     ugs_dll_path: Path,
-    cwd: Path = None
+    cwd: Optional[Path] = None
 ) -> Optional[Dict[str, str]]:
-    # Assumes UGS dependencies have been already setup (this assumption prevents us from redundently running `setup_dependencies()`)
+    # Assumes UGS dependencies have been already setup (prevents us from
+    # redundantly running `setup_dependencies()`)
 
     ugs_state = {}
+    status_args = ['dotnet', str(ugs_dll_path), 'status']
     try:
-        status_args = ['dotnet', ugs_dll_path, 'status']
-        with subprocess.Popen(status_args, cwd=cwd, stdin=subprocess.DEVNULL, stdout=subprocess.PIPE) as status_proc:
-            for staus_line in status_proc.stdout:
-                staus_line_str =  staus_line.decode().rstrip()
-                logging.debug(f'ugs> {staus_line_str}')
-                if 'Project:' in staus_line_str:
-                        match = re.match("Project:\s*//(\S+?)/(\S+)", staus_line_str)
-                        if match:
-                            ugs_state['client'] = match.group(1)
-                            ugs_state['project'] = match.group(2)
-                if 'User:' in staus_line_str:
-                        match = re.match("User:\s*(\S+)", staus_line_str)
-                        if match:
-                            ugs_state['user'] = match.group(1)
+        with subprocess.Popen(status_args, cwd=cwd, stdin=subprocess.DEVNULL,
+                              stdout=subprocess.PIPE,
+                              startupinfo=sb_utils.get_hidden_sp_startupinfo()
+                              ) as status_proc:
+            for status_line in status_proc.stdout:
+                status_line_str =  status_line.decode().rstrip()
+                logging.debug(f'ugs> {status_line_str}')
+                if 'Project:' in status_line_str:
+                    match = re.match(r"Project:\s*//(\S+?)/(\S+)",
+                                     status_line_str)
+                    if match:
+                        ugs_state['client'] = match.group(1)
+                        ugs_state['project'] = match.group(2)
+                if 'User:' in status_line_str:
+                    match = re.match(r"User:\s*(\S+)", status_line_str)
+                    if match:
+                        ugs_state['user'] = match.group(1)
     except Exception as exc:
-        logging.error(f'_get_ugs_state(): exception running Popen: {status_args}', exc_info=exc)
+        logging.error('_get_ugs_state(): exception running Popen: '
+                      f'{status_args}', exc_info=exc)
         return None
 
     if not ugs_state:
         return None
+
     return ugs_state
 
+
 def _set_active_ugs_context(
     ugs_dll_path: Path,
     ugs_settings: Dict[str, str],
-    cwd: Path = None
+    cwd: Optional[Path] = None
 ) -> int:
-    # Assumes UGS dependencies have been already setup (this assumption prevents us from redundently running `setup_dependencies()`)
+    # Assumes UGS dependencies have been already setup (prevents us from
+    # redundantly running `setup_dependencies()`)
 
+    init_args = ['dotnet', str(ugs_dll_path), 'init']
     try:
-        init_args = ['dotnet', ugs_dll_path, 'init']
-        for key,val in ugs_settings.items():
+        for key, val in ugs_settings.items():
             init_args.append(f'-{key}={val}')
 
-        with subprocess.Popen(init_args, cwd=cwd, stdin=subprocess.DEVNULL, stdout=subprocess.PIPE) as init_proc:
+        with subprocess.Popen(init_args, cwd=cwd, stdin=subprocess.DEVNULL,
+                              stdout=subprocess.PIPE,
+                              startupinfo=sb_utils.get_hidden_sp_startupinfo()
+                              ) as init_proc:
             for init_line in init_proc.stdout:
                 logging.info(f'ugs> {init_line.decode().rstrip()}')
     except Exception as exc:
-        logging.error(f'_set_active_ugs_context(): exception running Popen: {init_args}', exc_info=exc)
+        logging.error('_set_active_ugs_context(): exception running Popen: '
+                      f'{init_args}', exc_info=exc)
         return -1
 
     return 0
 
+
 def run(
-    ugs_args: List[str],
+    ugs_args: Union[str, List[str]],
     uproj_path: Path,
     ugs_bin_dir: Optional[Path] = None,
     user: Optional[str] = None,
     client: Optional[str] = None,
     output_handling_fn: Optional[Callable[[str], int]] = None,
 ) -> int:
-
-    # Even though we specify this should be a list, handle the simplified case of users passing a single word command
-    if type(ugs_args) == str:
+    # Even though we specify this should be a list, handle the simplified case
+    # of users passing a single word command
+    if isinstance(ugs_args, str):
         ugs_args = ugs_args.split()
     args_str = ' '.join(ugs_args)
 
     ugs_bin_path = find_bin(ugs_bin_dir)
     if not ugs_bin_path:
-        logging.error(f"Failed to find 'ugs.dll'. Aborting UnrealGameSync command: `ugs {args_str}`")
+        logging.error("Failed to find 'ugs.dll'. Aborting UnrealGameSync "
+                      f"command: `ugs {args_str}`")
         return -1
 
     ue_engine_dir = _find_engine_dir(uproj_path)
     if not ue_engine_dir:
-        logging.warning("Failed to locate the Unreal '/Engine/' directory. As a result we are unable to validate UGS dependencies. Ignoring and continuing on as if they were setup.")
+        logging.warning("Failed to locate the Unreal '/Engine/' directory. "
+                        "As a result we are unable to validate UGS "
+                        "dependencies. Ignoring and continuing on as if they "
+                        "were setup.")
     # Make sure the needed dependencies (dotnet, etc.) are installed upfront
     elif setup_dependencies(ue_engine_dir) != 0:
-        logging.warning('Failed to validate UnrealGameSync dependencies. Ignoring and continuing on as if they were setup.')
+        logging.warning('Failed to validate UnrealGameSync dependencies. '
+                        'Ignoring and continuing on as if they were setup.')
 
     cwd = uproj_path.parent
 
     logging.info("Capturing UnrealGameSync's current state.")
     ugs_state_to_restore = _get_active_ugs_context(ugs_bin_path, cwd=cwd)
     if not ugs_state_to_restore:
-        logging.warning('Failed to capture the current state of UnrealGameSync. Will not be able to restore it after completing our operation.')
+        logging.warning('Failed to capture the current state of '
+                        'UnrealGameSync. Will not be able to restore it after '
+                        'completing our operation.')
 
     logging.info('Setting UnrealGameSync context.')
     # UGS expects a project path relative to the repo's root
-    sanitized_proj_path = path.relpath(uproj_path, ue_engine_dir.parent) if ue_engine_dir else uproj_path
+    if ue_engine_dir:
+        sanitized_proj_path = path.relpath(uproj_path, ue_engine_dir.parent)
+    else:
+        sanitized_proj_path = uproj_path
     # and only accepts paths with delimited by forward-slashes
-    sanitized_proj_path = sanitized_proj_path.replace('\\','/')
-    op_context_params = { 'project' : sanitized_proj_path }
+    sanitized_proj_path = sanitized_proj_path.replace('\\', '/')
+    op_context_params = {'project': sanitized_proj_path}
     if user:
         op_context_params['user'] = user
     if client:
         op_context_params['client'] = client
 
-    init_result = _set_active_ugs_context(ugs_bin_path, op_context_params, cwd=cwd)
+    init_result = _set_active_ugs_context(ugs_bin_path, op_context_params,
+                                          cwd=cwd)
     if init_result != 0:
-        logging.error(f"Failed to initialize UnrealGameSync. Aborting UnrealGameSync command: `ugs {args_str}`")
+        logging.error("Failed to initialize UnrealGameSync. Aborting "
+                      f"UnrealGameSync command: `ugs {args_str}`")
         return init_result
 
     logging.info(f'Executing UnrealGameSync command: `ugs {args_str}`.')
     ugs_cmd = ['dotnet', ugs_bin_path] + ugs_args
 
     try:
-        with subprocess.Popen(ugs_cmd, cwd=cwd, stdin=subprocess.DEVNULL, stdout=subprocess.PIPE) as ugs_proc:
+        with subprocess.Popen(ugs_cmd, cwd=cwd, stdin=subprocess.DEVNULL,
+                              stdout=subprocess.PIPE,
+                              startupinfo=sb_utils.get_hidden_sp_startupinfo()
+                              ) as ugs_proc:
             for line in ugs_proc.stdout:
                 line_str = f'{line.decode().rstrip()}'
 
-                invalid_arg_match = re.match("Invalid argument:\s*(\S+)", line_str)
+                invalid_arg_match = re.match(r"Invalid argument:\s*(\S+)",
+                                             line_str)
                 if invalid_arg_match:
                     logging.error(f'ugs> {line_str}')
-                    logging.error(f'The current version of UnrealGameSync does not support the `{invalid_arg_match.match.group(1)}` argument. Please make sure UGS is up to date.')
+                    logging.error('The current version of UnrealGameSync does '
+                                  'not support the '
+                                  f'`{invalid_arg_match.group(1)}`'
+                                  'argument. Please ensure UGS is up to date.')
                     ugs_proc.terminate()
                 elif output_handling_fn:
                     if output_handling_fn(line_str) != 0:
@@ -197,16 +258,19 @@ def run(
                 else:
                     logging.info(f'ugs> {line_str}')
     except Exception as exc:
-        logging.error(f'ugs.run(): exception running Popen: {ugs_cmd}', exc_info=exc)
+        logging.error(f'ugs.run(): exception running Popen: {ugs_cmd}',
+                      exc_info=exc)
         return -1
 
     if ugs_state_to_restore:
         logging.info('Restoring previous UnrealGameSync state.')
-        if _set_active_ugs_context(ugs_bin_path, ugs_state_to_restore, cwd=cwd) != 0:
-            logging.warn("Failed to resotre UnrealGameSync's state prior to operation.")
+        if _set_active_ugs_context(ugs_bin_path,
+                                   ugs_state_to_restore, cwd=cwd) != 0:
+            logging.warn("Failed to restore UnrealGameSync's state.")
 
     return ugs_proc.returncode
 
+
 def sync(
     uproj_path: Path,
     sync_cl: Optional[int] = None,
@@ -233,17 +297,19 @@ def sync(
         client=client
     )
 
+
 def latest_chagelists(
     uproj_path: Path,
     ugs_bin_dir: Optional[Path] = None,
     user: Optional[str] = None,
     client: Optional[str] = None
 ) -> Optional[List[int]]:
-    
+
     changelist_list = []
 
-    def changes_output_handler(output_str:str, changes_out=changelist_list) -> int:
-        cl_desc_match = re.match("^\s*(\d+).*", output_str)
+    def changes_output_handler(output_str: str,
+                               changes_out=changelist_list) -> int:
+        cl_desc_match = re.match(r"^\s*(\d+).*", output_str)
         if cl_desc_match:
             changes_out.append(cl_desc_match.group(1))
         return 0
@@ -251,13 +317,142 @@ def latest_chagelists(
     run_result = run(
         'changes',
         uproj_path,
-        ugs_bin_dir =ugs_bin_dir,
-        user = user,
+        ugs_bin_dir=ugs_bin_dir,
+        user=user,
         client=client,
-        output_handling_fn = changes_output_handler
+        output_handling_fn=changes_output_handler
     )
 
     if run_result != 0:
         return None
 
-    return changelist_list
\ No newline at end of file
+    return changelist_list
+
+
+def get_depot_config_paths(
+    engine_dir: Union[str, Path],
+    project_dir: Union[str, Path]
+) -> List[str]:
+    # See GetDepotConfigPaths in UnrealGameSyncShared/Utility.cs
+    engine_dir = str(engine_dir)
+    project_dir = str(project_dir)
+
+    config_paths: List[str] = []
+
+    def add_platform_dirs(base: str, filename: str):
+        config_paths.extend([
+            f'{base}/{filename}',
+            f'{base}/*/{filename}'
+        ])
+
+    def add_platform_extensions(base: str, rel: str, filename: str):
+        add_platform_dirs(f'{base}{rel}', filename)
+        add_platform_dirs(f'{base}/Platforms/*{rel}', filename)
+        add_platform_dirs(f'{base}/Restricted/*{rel}', filename)
+
+    add_platform_extensions(engine_dir, '/Programs/UnrealGameSync',
+                            'UnrealGameSync.ini')
+
+    add_platform_extensions(project_dir, '/Build',
+                            'UnrealGameSync.ini')
+
+    return config_paths
+
+
+class IniParser:
+    '''
+    Modeled after ConfigParser, but handles Unreal duplicate `+Key=ValToAppend`
+    '''
+
+    class Section:
+        def __init__(self, name: str):
+            self.name = name
+            self.pairs: Dict[str, str] = {}
+
+        def set_value(self, key: str, value: str):
+            self.pairs[key] = value
+
+        def append_value(self, key: str, value: str):
+            if key in self.pairs:
+                self.pairs[key] += f'\n{value}'
+            else:
+                self.pairs[key] = value
+
+    def __init__(self):
+        self.sections: Dict[str, IniParser.Section] = {}
+
+    def read_string(self, string: str, source: str):
+        sfile = io.StringIO(string)
+        self.read_file(sfile, source)
+
+    def read_file(self, file: TextIO, source: Optional[str] = None):
+        if source is None:
+            try:
+                source = file.name
+            except AttributeError:
+                source = '<???>'
+
+        self._read(file, source)
+
+    def try_get(self, section_name: str, key: str) -> Optional[List[str]]:
+        section = self.sections.get(section_name)
+        if section is None:
+            return None
+        value = section.pairs.get(key)
+        if value is None:
+            return None
+        return value.splitlines()
+
+    def _read(self, file: TextIO, source: str):
+        # See Parse() in UnrealGameSyncShared/ConfigFile.cs
+        current_section: Optional[IniParser.Section] = None
+        for line_num, line in enumerate(file, start=1):
+            line = line.rstrip('\r\n')
+            strip_line = line.strip()
+
+            if len(strip_line) == 0:
+                continue
+
+            if strip_line.startswith(';'):
+                continue
+
+            if strip_line.startswith('[') and strip_line.endswith(']'):
+                section_name = strip_line[1:-1]
+                current_section = self.sections.get(section_name)
+                if current_section is None:
+                    new_section = IniParser.Section(section_name)
+                    self.sections[section_name] = new_section
+                    current_section = new_section
+                continue
+
+            if current_section is None:
+                logging.warning(f'{source}:{line_num}: error: '
+                                'current_section is None')
+                logging.warning(f'\t{line}')
+                continue
+
+            equals_idx = strip_line.find('=')
+            if equals_idx != -1:
+                value = line[equals_idx+1:].lstrip()
+                if strip_line.startswith('+'):
+                    key = strip_line[1:equals_idx].strip()
+                    current_section.append_value(key, value)
+                else:
+                    key = strip_line[0:equals_idx].rstrip()
+                    current_section.set_value(key, value)
+            else:
+                logging.warning(f'{source}:{line_num}: error: missing =')
+                logging.warning(f'\t{line}')
+
+    def dump(self):
+        ''' Write parsed state to `logging.info`. '''
+        for name, section in self.sections.items():
+            logging.info(f'[{name}]')
+            for key, value in section.pairs.items():
+                value_lines = value.splitlines()
+                if len(value_lines) == 1:
+                    logging.info(f'{key}={value}')
+                else:
+                    logging.info(f'{key}=')
+                    for line in value_lines:
+                        logging.info(f'\t{line}')
