Switchboard: Changed default folder when browsing for new nDisplay config.

Instead of a global last browed location setting, it now prefers the folder of the current nDisplay config. If there isn't one, it prefers the content folder of the current project. If that fails, it uses the user's home folder.

#skipci

diff --git a/Engine/Plugins/VirtualProduction/Switchboard/Source/Switchboard/switchboard/devices/ndisplay/plugin_ndisplay.py b/Engine/Plugins/VirtualProduction/Switchboard/Source/Switchboard/switchboard/devices/ndisplay/plugin_ndisplay.py
--- a/Engine/Plugins/VirtualProduction/Switchboard/Source/Switchboard/switchboard/devices/ndisplay/plugin_ndisplay.py
+++ b/Engine/Plugins/VirtualProduction/Switchboard/Source/Switchboard/switchboard/devices/ndisplay/plugin_ndisplay.py
@@ -153,20 +153,33 @@ class AddnDisplayDialog(AddDeviceDialog):
     def on_clicked_btnBrowse(self):
         ''' Opens a file dialog to browse for the config file
         '''
-        start_path = str(Path.home())
 
-        if (SETTINGS.LAST_BROWSED_PATH and
-                os.path.exists(SETTINGS.LAST_BROWSED_PATH)):
-            start_path = SETTINGS.LAST_BROWSED_PATH
+        start_path = ''
 
+        # prefer starting at the folder of the last ndisplay config
+        last_config_path = DevicenDisplay.csettings['ndisplay_config_file'].get_value()
+
+        if os.path.exists(last_config_path):
+            start_path = os.path.split(last_config_path)[0]
+
+        # otherwise, use the project's Content folder
+        if not start_path:
+            content_path = os.path.join(os.path.split(CONFIG.UPROJECT_PATH.get_value())[0],'Content')
+            if os.path.exists(content_path):
+                start_path = content_path
+
+        # our fallback is the home folder
+        if not start_path:
+            start_path = str(Path.home())
+
+        # show the open file dialog
         cfg_path, _ = QtWidgets.QFileDialog.getOpenFileName(
             self, "Select nDisplay config file", start_path,
             "nDisplay Config (*.ndisplay *.uasset)")
 
+        # update the field with the selected path
         if len(cfg_path) > 0 and os.path.exists(cfg_path):
             self.cbConfigs.setCurrentText(cfg_path)
-            SETTINGS.LAST_BROWSED_PATH = os.path.dirname(cfg_path)
-            SETTINGS.save()
 
     def generate_short_unique_config_name(config_path: str, file_name: str) -> str:
         config_path = CONFIG.shrink_path(config_path)
