Switchboard: Fixed parse assets called when pressing enter on unrelated fields in the settings dialog.

The issue was worked around by implementing a custom click event on the push button.

#skipci

diff --git a/Engine/Plugins/VirtualProduction/Switchboard/Source/Switchboard/switchboard/devices/unreal/plugin_unreal.py b/Engine/Plugins/VirtualProduction/Switchboard/Source/Switchboard/switchboard/devices/unreal/plugin_unreal.py
--- a/Engine/Plugins/VirtualProduction/Switchboard/Source/Switchboard/switchboard/devices/unreal/plugin_unreal.py
+++ b/Engine/Plugins/VirtualProduction/Switchboard/Source/Switchboard/switchboard/devices/unreal/plugin_unreal.py
@@ -313,8 +313,10 @@ class LiveLinkPresetSetting(Setting):
         btn_refresh.setIcon(QtGui.QIcon(':/icons/images/icon_refresh.png'))
         btn_refresh.setProperty("frameless", True)
         btn_refresh.setToolTip("Refresh list")
+        btn_refresh.setDefault(False)
+        btn_refresh.setAutoDefault(False)
         btn_refresh.clicked.connect(
-            lambda: self._refresh_livelink_preset_list(combo, override_device_name)
+            lambda: self._refresh_list(combo, override_device_name)
         )
 
         # create layout with combo and refresh button
@@ -329,7 +331,7 @@ class LiveLinkPresetSetting(Setting):
 
         return layout
 
-    def _refresh_livelink_preset_list(self, combo:QtWidgets.QComboBox, override_device_name:Optional[str] = None):
+    def _refresh_list(self, combo:QtWidgets.QComboBox, override_device_name:Optional[str] = None):
         ''' Traveses the content to refresh the list of live link presets available '''
 
         DeviceUnreal.analyze_project_assets()
