Switchboard: Preliminary support for UGS sync filters via Python sync. (Currently relies on P4Python being manually added to the device's Python environment.)

- Moved `p4_*` functions from sbl_helper.py to p4_utils.py, and UGS config depot retrieval/parsing to ugs_utils.py.
- Added some basic unit tests for ugs_utils.py.
- Made OptionSetting/MultiOptionSetting/MultiSelectionComboBox usable with complex value types (instead of just strings).
- Added a means of customizing Setting (JSON) de/serialization for those complex value types. (Uses fluent/builder pattern rather than proliferating more initializer arguments.)
- Modified switchboard_logging.py to use the root logger for compatibility with library code that calls standard `logging.*` methods.
- Avoid writing a corrupt truncated config file to disk in the event of an exception during json.dump.
- Updated some type annotations for PEP585 (since 3.9, list[T] instead of typing.List[T]).

diff --git a/Engine/Plugins/VirtualProduction/Switchboard/Source/Switchboard/sbl_helper.py b/Engine/Plugins/VirtualProduction/Switchboard/Source/Switchboard/sbl_helper.py
--- a/Engine/Plugins/VirtualProduction/Switchboard/Source/Switchboard/sbl_helper.py
+++ b/Engine/Plugins/VirtualProduction/Switchboard/Source/Switchboard/sbl_helper.py
@@ -17,231 +17,16 @@ import sys
 import threading
 import tempfile
 import traceback
-from typing import Callable, Dict, List, Optional, Sequence, Union
-from switchboard import ugs_utils as UGS
+from typing import Callable, Optional, Sequence, Union
+import uuid
+
+from switchboard import p4_utils, ugs_utils as UGS
 
 
 MAX_PARALLEL_WORKERS = 8  # Consistent with UGS, ushell, etc.
 
 SCRIPT_NAME = os.path.basename(__file__)
-__version__ = '1.1.0'
-
-
-def p4_variables() -> Dict[str, str]:
-    '''
-    Returns Perforce variable (e.g. `P4CONFIG`, `p4 set`) key/value pairs.
-    '''
-    output = subprocess.check_output(['p4', 'set', '-q']).decode()
-    variables: Dict[str, str] = {}
-    for line in output.splitlines():
-        var_name, _, var_value = line.partition('=')
-        variables[var_name] = var_value
-    return variables
-
-
-def p4(
-    cmd: str,
-    cmd_opts: List[str],
-    global_opts: Optional[List[str]] = None,
-    *,
-    user: Optional[str] = None,
-    client: Optional[str] = None,
-) -> subprocess.Popen:
-    ''' Runs a p4 command and includes some common/required switches. '''
-    args = ['p4', f'-zprog={SCRIPT_NAME}', f'-zversion={__version__}', '-ztag',
-            '-G', '-Qutf8']
-
-    if user is not None:
-        args.extend(['-u', user])
-    if client is not None:
-        args.extend(['-c', client])
-
-    args.extend(global_opts or [])
-
-    args.append(cmd)
-    args.extend(cmd_opts)
-
-    logging.debug(f'p4(): invoking subprocess: args={args}')
-
-    # stdin=subprocess.DEVNULL is required when launched via
-    # SwitchboardListener; otherwise this call tries to make the non-existent
-    # stdin inheritable, raising `OSError: [WinError 6] The handle is invalid`
-    return subprocess.Popen(args, stdout=subprocess.PIPE,
-                            stdin=subprocess.DEVNULL)
-
-
-def p4_get_records(
-    cmd: str,
-    opts: List[str],
-    global_opts: Optional[List[str]] = None,
-    *,
-    include_info: bool = False,
-    include_error: bool = False,
-    user: Optional[str] = None,
-    client: Optional[str] = None,
-) -> List[dict]:
-    results: List[dict] = []
-    proc = p4(cmd, opts, global_opts=global_opts, user=user, client=client)
-    while True:
-        try:
-            record = marshal.load(proc.stdout)
-        except EOFError:
-            break
-
-        # Possible codes: stat, info, error, text (only `p4 print`)
-        include_record = False
-        if (record[b'code'] == b'stat') or (record[b'code'] == b'text'):
-            include_record = True
-        elif include_info and (record[b'code'] == b'info'):
-            include_record = True
-        elif include_error and (record[b'code'] == b'error'):
-            include_record = True
-
-        if include_record:
-            results.append(record)
-
-    proc.stdout.close()
-    return results
-
-
-def p4_changes(
-    pathspecs: List[str],
-    *,
-    limit: int = 10,
-    user: Optional[str] = None,
-    client: Optional[str] = None,
-) -> List[dict]:
-    limit = min(limit, 100)
-    opts = ['-s', 'submitted', '-t', '-L', f'-m{limit}', *pathspecs]
-    return p4_get_records('changes', opts, user=user, client=client)
-
-
-def p4_fstat(
-    pathspecs: List[str],
-    *,
-    limit: int = 0,
-    user: Optional[str] = None,
-    client: Optional[str] = None,
-) -> List[dict]:
-    opts: List[str] = []
-
-    if limit > 0:
-        opts.append(f'-m{limit}')
-
-    opts.extend(pathspecs)
-
-    return p4_get_records('fstat', opts, user=user, client=client)
-
-
-def p4_have(
-    pathspecs: List[str],
-    *,
-    user: Optional[str] = None,
-    client: Optional[str] = None,
-) -> List[Optional[int]]:
-    '''
-    Returns a list of revision numbers, or `None` if the file doesn't exist on
-    the client.
-    '''
-    result: List[Optional[int]] = []
-    records = p4_get_records('have', pathspecs, include_error=True, user=user,
-                             client=client)
-    for record in records:
-        if record[b'code'] == b'stat':
-            result.append(int(record[b'haveRev']))
-        else:
-            result.append(None)
-    return result
-
-
-@dataclass
-class P4PrintResult:
-    stat: dict
-    contents: Optional[bytearray]
-
-    @property
-    def is_valid(self) -> bool:
-        try:
-            return (self.stat[b'code'] == b'stat') and (self.contents is not None)
-        except KeyError:
-            return False
-
-    @property
-    def text(self) -> Optional[str]:
-        if self.contents is not None:
-            return self.contents.decode()
-        else:
-            return None
-
-
-def p4_print(
-    pathspecs: List[str],
-    *,
-    user: Optional[str] = None,
-    client: Optional[str] = None,
-) -> List[P4PrintResult]:
-    '''
-    Returns a `P4PrintResult` for each path in `pathspecs`.
-    If the file does not exist, or is deleted at the specified revision,
-    fileContents will be `None`.
-    '''
-    results: List[P4PrintResult] = []
-    records = p4_get_records('print', pathspecs, include_info=True,
-                             include_error=True, user=user, client=client)
-
-    # For each file, a stat record is returned, then zero or more text records.
-    # Cases where there's no text record include errors or deleted files.
-    for record in records:
-        if (record[b'code'] == b'stat') or (record[b'code'] == b'error'):
-            results.append(P4PrintResult(record, None))
-            continue
-        elif record[b'code'] == b'text':
-            current_result = results[-1]
-            if current_result.contents is None:
-                current_result.contents = bytearray()
-            current_result.contents.extend(record[b'data'])
-            continue
-
-        logging.warning(f'p4_print(): Unhandled record: {record}')
-
-    return results
-
-
-def p4_latest_code_change(
-    paths: List[str],
-    *,
-    range: Optional[str] = None,
-    exts: Optional[List[str]] = None,
-    user: Optional[str] = None,
-    client: Optional[str] = None,
-) -> Optional[int]:
-    '''
-    Code CL determination compatible with UGS/precompiled binaries.
-    See `WorkspaceUpdate.ExecuteAsync` in `WorkspaceUpdate.cs`
-    '''
-
-    ALL_CODE_EXTS = [
-        '.c', '.cc', '.cpp', '.inl', '.m', '.mm', '.rc', '.cs', '.csproj',
-        '.h', '.hpp', '.usf', '.ush', '.uproject', '.uplugin', '.sln']
-
-    if exts is None:
-        exts = ALL_CODE_EXTS
-
-    rangespec = f'@{range}' if range is not None else ''
-
-    code_paths: List[str] = []
-    for path in paths:
-        if path.endswith('...'):
-            code_paths.extend([f'{path}/*{ext}{rangespec}'
-                                for ext in exts])
-
-    code_changes = p4_changes(code_paths, limit=1,
-                              user=user, client=client)
-
-    if len(code_changes) > 0:
-        return max(int(x[b'change']) for x in code_changes)
-    else:
-        return None
+__version__ = '1.2.0'
 
 
 class SyncFile:
@@ -276,9 +61,9 @@ class SyncWorkload:
     Represents a sync operation, or e.g. a parallel worker's subset thereof.
     '''
     def __init__(self):
-        self.sync_files: Dict[str, SyncFile] = {}
+        self.sync_files: dict[str, SyncFile] = {}
         self.workload_size = 0
-        self.client_path_to_depot_path: Dict[str, str] = {}
+        self.client_path_to_depot_path: dict[str, str] = {}
 
     def add_or_replace(self, new_file: SyncFile):
         if new_file.depot_path in self.sync_files:
@@ -346,7 +131,7 @@ def p4_sync_worker(
         argfile.write(f'{item.spec}\n')
     argfile.close()  # Otherwise the p4 process can't access the file.
 
-    proc = p4('sync', opts, gopts, user=user, client=client)
+    proc = p4_utils.p4('sync', opts, gopts, user=user, client=client)
 
     remaining_work = copy.deepcopy(sync_workload)
 
@@ -409,14 +194,14 @@ def p4_sync_worker(
 
 
 def p4_get_preview_sync_files(
-    input_specs: List[str],
+    input_specs: list[str],
     *,
     user: Optional[str] = None,
     client: Optional[str] = None,
 ) -> SyncWorkload:
     ''' Runs a preview sync to generate a list of files and revisions. '''
     result_workload = SyncWorkload()
-    records = p4_get_records(
+    records = p4_utils.p4_get_records(
         'sync', ['-n', *input_specs], include_info=True, include_error=True,
         user=user, client=client)
 
@@ -468,7 +253,7 @@ def friendly_bytes(size: float) -> str:
 
 
 def p4_sync(
-    input_specs: List[str],
+    input_specs: list[str],
     num_workers: int = MAX_PARALLEL_WORKERS,
     *,
     sync_filter_fn: Optional[Callable[[SyncWorkload], None]] = None,
@@ -510,6 +295,8 @@ def p4_sync(
         else:
             nondelete_work.add(item)
 
+    del total_work
+
     # First pass: Attempt to run all deletes (in the main thread).
     total_delete_files = len(delete_work.sync_files)
     if total_delete_files > 0:
@@ -635,7 +422,7 @@ def dispatch_sync_workers(
     completion_queue: Queue[SyncFile] = Queue()
 
     # Summarize work distribution and start workers.
-    workers: List[threading.Thread] = []
+    workers: list[threading.Thread] = []
     for i in range(num_workers):
         worker_name = f'Worker {i}'
         load = worker_loads[i]
@@ -703,7 +490,7 @@ class SbListenerHelper:
 
         self.uproj_path: Optional[pathlib.Path] = None
         self.engine_dir: Optional[pathlib.Path] = None
-        self.additional_paths_to_sync: List[str] = []
+        self.additional_paths_to_sync: list[str] = []
         self.sync_engine_cl: Optional[int] = None
         self.sync_project_cl: Optional[int] = None
         self.clobber_engine: bool = False
@@ -713,6 +500,12 @@ class SbListenerHelper:
         self.dry_run: bool = False
         self.p4user: Optional[str] = None
         self.p4client: Optional[str] = None
+        self.p4clientspec: dict[str, str] = {}
+
+        if UGS.SyncFilters.supported():
+            self.ugs_filters = UGS.SyncFilters()
+        else:
+            self.ugs_filters = None
 
     @staticmethod
     def build_parser():
@@ -745,18 +538,6 @@ class SbListenerHelper:
         sync_parser.add_argument(
             '--generate', action='store_true',
             help='Run GenerateProjectFiles after syncing')
-        sync_parser.add_argument(
-            '--use-ugs', action='store_true',
-            help='Specifies that UnrealGameSync should be used to sync the '
-                 'project and engine together.')
-        sync_parser.add_argument(
-            '--use-pcbs', action='store_true',
-            help='Specifies that precompiled binaries should be synced along '
-                 'with the project (requires syncing via UnrealGameSync).')
-        sync_parser.add_argument(
-            '--ugs-lib-dir', type=pathlib.Path,
-            help='Directory path specifying where to find the UnrealGameSync '
-                 'library (ugs.dll).')
         sync_parser.add_argument(
             '--clobber-engine', action='store_true',
             help='Override noclobber for engine files')
@@ -771,6 +552,32 @@ class SbListenerHelper:
         sync_parser.add_argument('-c', '--p4client', type=str,
                                  help='Override Perforce P4CLIENT')
 
+        # UGS sync filter options
+        sync_parser.add_argument(
+            '--include-categories', type=str,
+            help='List of UGS SyncCategory UUIDs to include in sync')
+        sync_parser.add_argument(
+            '--custom-view', type=str,
+            help='Comma-separated Perforce wildcards to exclude from sync')
+        sync_parser.add_argument(
+            '--remove-excluded-workspace-files', action='store_true',
+            help='Scan workspace to remove all files excluded by sync filter')
+
+        # UGS CLI integration options
+        sync_parser.add_argument(
+            '--use-ugs', action='store_true',
+            help='Specifies that UnrealGameSync should be used to sync the '
+                 'project and engine together.')
+        sync_parser.add_argument(
+            '--use-pcbs', action='store_true',
+            help='Specifies that precompiled binaries should be synced along '
+                 'with the project (requires syncing via UnrealGameSync).')
+        sync_parser.add_argument(
+            '--ugs-lib-dir', type=pathlib.Path,
+            help='Directory path specifying where to find the UnrealGameSync '
+                 'library (ugs.dll).')
+
+        # Debugging options
         sync_parser.add_argument(
             '--dump-sync', type=argparse.FileType('wt'),
             metavar='SYNC_WORKLOAD_TEXT_FILE',
@@ -793,12 +600,51 @@ class SbListenerHelper:
         assert False
 
     def run_sync(self, options: argparse.Namespace) -> int:
-        engine_paths: List[str] = []
-        project_paths: List[str] = []
-        pathspecs: List[str] = []
+        engine_paths: list[str] = []
+        project_paths: list[str] = []
+        pathspecs: list[str] = []
 
         self.check_sync_options(options)
-        self.additional_paths_to_sync = self.get_additional_paths_to_sync()
+
+        ugs_config = UGS.parse_depot_ugs_configs(
+            engine_dir=self.engine_dir, project_dir=self.uproj_path.parent,
+            p4user=self.p4user, p4client=self.p4client)
+
+        self.additional_paths_to_sync = ugs_config.try_get(
+            'Perforce', 'AdditionalPathsToSync') or []
+
+        # Parse sync filter categories from INI and match to command line args
+        any_active_filters = False
+        if options.include_categories or options.custom_view:
+            if not self.ugs_filters:
+                self.parser.error('Sync filters were specified, but filtering '
+                                  'is unavailable (missing P4Python?)')
+
+            if options.include_categories:
+                include_ids = set(
+                    uuid.UUID(id)
+                    for id in options.include_categories.split(','))
+                self.ugs_filters.read_categories_from_ini_parser(ugs_config)
+                for category in self.ugs_filters.categories.values():
+                    if category.id not in include_ids:
+                        logging.debug(f'Excluding sync filter "{category.name}"')
+                        self.ugs_filters.exclude_category(category.id)
+                        any_active_filters = True
+                    else:
+                        include_ids.remove(category.id)
+
+                # Any remaining IDs here were unmatched to categories
+                if len(include_ids):
+                    unknown = ', '.join(str(id) for id in include_ids)
+                    self.parser.error(
+                        f'Unknown --include-categories: {unknown}')
+
+            if options.custom_view:
+                custom_views: list[str] = options.custom_view.split(',')
+                for view in custom_views:
+                    self.ugs_filters.map.insert(view)
+
+            logging.debug(f'Sync filter: {self.ugs_filters.map.as_array()}')
 
         if self.sync_engine_cl:
             engine_paths.append(f'{self.engine_dir.parent / "*"}')
@@ -833,6 +679,25 @@ class SbListenerHelper:
                 if buildver_info.local_path and buildver_info.have_rev:
                     pathspecs.append(f"{buildver_info.local_path}#0")
 
+        remove_excluded_workload = SyncWorkload()
+        if options.remove_excluded_workspace_files and self.ugs_filters:
+            logging.info('Finding files in workspace that need to be removed.')
+            haves = p4_utils.p4_have(['//...'],
+                                     user=self.p4user, client=self.p4client)
+            stream = self.p4clientspec['Stream']
+            stream_path = pathlib.PurePosixPath(stream)
+            num_exclusions = 0
+            for have in haves:
+                client_str = have['path']
+                depot_str = have['depotFile']
+                depot_path = pathlib.PurePosixPath(depot_str)
+                rel_path = f'/{depot_path.relative_to(stream_path)}'
+                if not self.ugs_filters.includes_path(rel_path):
+                    num_exclusions += 1
+                    remove_excluded_workload.add(
+                        SyncFile(depot_str, client_str, 0, 0, 'deleted'))
+            logging.info(f'Found {num_exclusions:,} excluded files')
+
         def sync_filter(workload: SyncWorkload):
             # Don't add Build.version to the workspace, but do allow removal.
             buildver_path = buildver_info.depot_path
@@ -840,6 +705,26 @@ class SbListenerHelper:
                 if workload.sync_files[buildver_path].rev != 0:
                     workload.remove(buildver_path)
 
+            if any_active_filters and self.ugs_filters:
+                num_files = len(workload.sync_files)
+                logging.info(f'Applying sync filters... ({num_files:,} files)')
+                stream = self.p4clientspec['Stream']
+                stream_path = pathlib.PurePosixPath(stream)
+                excluded_files: set[str] = set()
+                for depot_str, syncfile in workload.sync_files.items():
+                    depot_path = pathlib.PurePosixPath(depot_str)
+                    depot_rel_path = f'/{depot_path.relative_to(stream_path)}'
+                    if not self.ugs_filters.includes_path(depot_rel_path):
+                        excluded_files.add(depot_str)
+
+                num_exclusions = len(excluded_files)
+                logging.info(f'Filter excluded {num_exclusions:,} files')
+                for exclusion in excluded_files:
+                    workload.sync_files.pop(exclusion)
+
+            for removal in remove_excluded_workload.sync_files.values():
+                workload.add_or_replace(removal)
+
         def clobber_filter(clobber: SyncWorkload):
             clobber.sync_files = dict(sorted(clobber.sync_files.items()))
 
@@ -921,8 +806,8 @@ class SbListenerHelper:
 
     def get_buildver_info(
         self,
-        engine_paths: List[str],
-        project_paths: List[str],
+        engine_paths: list[str],
+        project_paths: list[str],
     ) -> BuildVerInfo:
         '''
         Read `Build.version` from depot and return its info, depot contents,
@@ -943,12 +828,14 @@ class SbListenerHelper:
 
         p4opts = {'user': self.p4user, 'client': self.p4client}
 
-        ret_info.have_rev = p4_have([buildver_local_path], **p4opts)[0]
+        record = p4_utils.p4_have([buildver_local_path], **p4opts)[0]
+        ret_info.have_rev = int(record['haveRev']) if (
+            record['code'] == 'stat') else None
 
         if not self.sync_engine_cl:
             return ret_info
 
-        prints = p4_print(
+        prints = p4_utils.p4_print(
             [f'{path}@{self.sync_engine_cl}'
              for path in (epicint_local_path, buildver_local_path)],
             **p4opts)
@@ -972,7 +859,7 @@ class SbListenerHelper:
         eng_upper_bound_cl = self.sync_engine_cl
         proj_upper_bound_cl = self.sync_project_cl
 
-        def quick_narrow_lower_cl(paths: List[str], upper_cl: int,
+        def quick_narrow_lower_cl(paths: list[str], upper_cl: int,
                                   in_lower_cl: Optional[int]) -> Optional[int]:
             quick_lower = upper_cl - 1_000
             if in_lower_cl and in_lower_cl > quick_lower:
@@ -981,7 +868,7 @@ class SbListenerHelper:
                 return upper_cl
             if quick_lower > 0:
                 quick_range = f'{quick_lower},{upper_cl}'
-                check_lower = p4_latest_code_change(
+                check_lower = p4_utils.p4_latest_code_change(
                     paths, range=quick_range, exts=['.cpp'], **p4opts)
 
                 if check_lower:
@@ -996,7 +883,7 @@ class SbListenerHelper:
                                                eng_upper_bound_cl,
                                                lower_bound_cl)
 
-        if proj_upper_bound_cl:
+        if (lower_bound_cl or 0) < (proj_upper_bound_cl or 0):
             lower_bound_cl = quick_narrow_lower_cl(project_paths,
                                                    proj_upper_bound_cl,
                                                    lower_bound_cl)
@@ -1009,10 +896,20 @@ class SbListenerHelper:
             eng_range = f'{lower_bound_cl},{eng_upper_bound_cl}'
             proj_range = f'{lower_bound_cl},{proj_upper_bound_cl}'
 
-        eng_code_cl = p4_latest_code_change(engine_paths, range=eng_range,
-                                            **p4opts)
-        proj_code_cl = p4_latest_code_change(project_paths, range=proj_range,
-                                             **p4opts)
+        if lower_bound_cl and (lower_bound_cl >= eng_upper_bound_cl):
+            eng_code_cl = eng_upper_bound_cl
+        else:
+            eng_code_cl = p4_utils.p4_latest_code_change(
+                engine_paths, range=eng_range, **p4opts)
+
+        if proj_upper_bound_cl:
+            if lower_bound_cl and (lower_bound_cl >= proj_upper_bound_cl):
+                proj_code_cl = proj_upper_bound_cl
+            else:
+                proj_code_cl = p4_utils.p4_latest_code_change(
+                    project_paths, range=proj_range, **p4opts)
+        else:
+            proj_code_cl = None
 
         # These are the values we'll actually write to the JSON
         ver_cl = max(filter(None, [self.sync_engine_cl, self.sync_project_cl]))
@@ -1020,7 +917,7 @@ class SbListenerHelper:
 
         # Generate and return the updated Build.version contents.
         try:
-            build_ver: Dict[str, Union[int, str]] = json.loads(
+            build_ver: dict[str, Union[int, str]] = json.loads(
                 ret_info.depot_text)
         except json.JSONDecodeError:
             logging.error('Unable to parse depot Build.version JSON',
@@ -1043,78 +940,25 @@ class SbListenerHelper:
         if no_depot_compat or licensee_changed:
             updates['CompatibleChangelist'] = ver_compatible_cl
 
-        updated_build_ver = build_ver.copy()
-        updated_build_ver.update(updates)
-        if updated_build_ver == build_ver:
+        build_ver.update(updates)
+
+        # Read existing local file and skip update if unchanged
+        existing_build_ver = None
+        try:
+            with open(buildver_local_path, 'r') as f:
+                existing_build_ver = json.load(f)
+        except Exception as exc:
+            logging.warning('Exception reading local Build.version',
+                            exc_info=exc)
+
+        if build_ver == existing_build_ver:
             logging.info("get_buildver_info(): contents haven't changed")
         else:
             logging.info(f"get_buildver_info(): {updates}")
-            ret_info.updated_text = json.dumps(updated_build_ver, indent=4)
+            ret_info.updated_text = json.dumps(build_ver, indent=2)
 
         return ret_info
 
-    def get_additional_paths_to_sync(self) -> List[str]:
-        # We append depot wildcards to host paths and end up with mixed syntax
-        # here, but fstat separates them into normalized depot/host paths.
-        mixed_config_paths = UGS.get_depot_config_paths(self.engine_dir,
-                                                        self.uproj_path.parent)
-
-        p4opts = {'user': self.p4user, 'client': self.p4client}
-
-        config_fstats = p4_fstat(mixed_config_paths, **p4opts)
-
-        config_contents_map: Dict[str, Optional[str]] = {}
-
-        # First pass: get contents of locally modified files, or depot paths
-        # of files to `p4 print` in the second pass otherwise
-        for record in config_fstats:
-            action: Optional[bytes] = record.get(b'action')
-            head_action: Optional[bytes] = record.get(b'headAction')
-            if record[b'code'] == b'stat':
-                if (b'clientFile' in record) and (b'action' in record):
-                    # If the client file exists and is modified, use local ver
-                    client_path = pathlib.Path(record[b'clientFile'].decode())
-                    if client_path.exists and action == b'edit':
-                        local_contents = client_path.read_text()
-                        config_contents_map[str(client_path)] = local_contents
-                        continue
-                    elif action == b'delete':
-                        continue
-
-                if head_action == b'delete':
-                    continue
-
-                if b'depotFile' in record:
-                    depot_file = record[b'depotFile'].decode()
-                    config_contents_map[depot_file] = None
-                    continue
-
-            logging.warning(f'Unhandled fstat record: {record}')
-
-        # Second pass: p4 prints of depot configs @ head
-        depot_print_paths = [k for k, v in config_contents_map.items()
-                             if v is None]  # configs we didn't read locally
-
-        config_prints = p4_print(depot_print_paths, **p4opts)
-
-        for depot_path, record in zip(depot_print_paths, config_prints):
-            if record.is_valid:
-                config_contents_map[depot_path] = record.text
-            else:
-                logging.warning(f'Invalid P4PrintResult for {depot_path}')
-
-        parser = UGS.IniParser()
-        for path, contents in config_contents_map.items():
-            if contents is None:
-                logging.warning(f'No contents for {path}')
-                continue
-
-            parser.read_string(contents, path)
-
-        paths = parser.try_get('Perforce', 'AdditionalPathsToSync')
-        logging.debug(f'AdditionalPathsToSync: {paths}')
-        return paths or []
-
     def on_sync_progress(
         self,
         completed_bytes: int,
@@ -1145,12 +989,16 @@ class SbListenerHelper:
         between related options.
         '''
         self.dry_run = options.dry_run
-        self.p4user = options.p4user
-        self.p4client = options.p4client
         self.clobber_engine = options.clobber_engine
         self.clobber_project = options.clobber_project
         self.ugs_lib_dir = options.ugs_lib_dir
 
+        self.p4user = options.p4user
+        self.p4client = options.p4client
+        self.p4clientspec = p4_utils.p4_get_client(self.p4client,
+                                                   user=self.p4user,
+                                                   client=self.p4client) or {}
+
         if options.project is None:
             if options.project_cl is not None:
                 self.parser.error('--project-cl requires --project')
@@ -1261,6 +1109,9 @@ class SbListenerHelper:
 
 
 def main() -> int:
+    p4_utils.meta_zprog = SCRIPT_NAME
+    p4_utils.meta_zversion = __version__
+
     app = SbListenerHelper()
     result = app.run()
     logging.info(f'Finished with exit code {result}')
diff --git a/Engine/Plugins/VirtualProduction/Switchboard/Source/Switchboard/switchboard/config.py b/Engine/Plugins/VirtualProduction/Switchboard/Source/Switchboard/switchboard/config.py
--- a/Engine/Plugins/VirtualProduction/Switchboard/Source/Switchboard/switchboard/config.py
+++ b/Engine/Plugins/VirtualProduction/Switchboard/Source/Switchboard/switchboard/config.py
@@ -2,14 +2,14 @@
 
 import collections
 import fnmatch
+import io
 import json
 import os
 import pathlib
 import shutil
 import socket
 import sys
-import typing
-from typing import Type
+from typing import Any, Callable, Optional, Tuple, Type, Union
 from enum import Enum
 
 from PySide2 import QtCore
@@ -35,7 +35,7 @@ DEFAULT_MAP_TEXT = '-- Default Map --'
 ENABLE_UGS_SUPPORT = False
 
 
-def migrate_comma_separated_string_to_list(value) -> typing.List[str]:
+def migrate_comma_separated_string_to_list(value) -> list[str]:
     if isinstance(value, str):
         return value.split(",")
     # Technically we should check whether every element is a string but we skip it here
@@ -70,10 +70,10 @@ class Setting(QtCore.QObject):
         attr_name: str,
         nice_name: str,
         value,
-        tool_tip: typing.Optional[str] = None,
+        tool_tip: Optional[str] = None,
         show_ui: bool = True,
         allow_reset: bool = True,
-        migrate_data: typing.Callable[[any], None] = None
+        migrate_data: Optional[Callable[[Any], None]] = None,
     ):
         '''
         Create a new Setting object.
@@ -109,7 +109,7 @@ class Setting(QtCore.QObject):
         self._base_widget = None
         self._override_widgets = {}
         self._on_setting_changed_lambdas = {}
-        
+
         # Appears when override value is different from _value
         self._allow_reset = allow_reset
         self._base_reset_widget = None
@@ -117,9 +117,20 @@ class Setting(QtCore.QObject):
 
         self._migrate_data = migrate_data
 
+        self._get_json_override_fn: Optional[Callable[[Any], Any]] = None
+        self._config_set_override_fn: Optional[Callable[[Any], Any]] = None
+
         self.tool_tip = tool_tip
         self.show_ui = show_ui
 
+    def with_get_json_override_fn(self, fn: Callable[[Any], Any]):
+        self._get_json_override_fn = fn
+        return self
+
+    def with_config_set_override_fn(self, fn: Callable[[Any], Any]):
+        self._config_set_override_fn = fn
+        return self
+
     def is_overridden(self, device_name: str) -> bool:
         try:
             return self._overrides[device_name] != self._value
@@ -141,7 +152,7 @@ class Setting(QtCore.QObject):
 
         self.signal_setting_changed.emit(old_value, self._value)
         self._refresh_reset_base_widget()
-    
+
     def override_value(self, device_name: str, override):
         override = self._filter_value(override)
 
@@ -154,16 +165,38 @@ class Setting(QtCore.QObject):
 
         self._refresh_reset_override_widget(device_name)
 
-    def get_value(self, device_name: typing.Optional[str] = None):
-        def get_value(self, device_name: typing.Optional[str] = None):
+    def get_value(self, device_name: Optional[str] = None):
+        def get_value(self, device_name: Optional[str] = None):
             try:
                 return self._overrides[device_name]
             except KeyError:
                 return self._value
-            
+
         value = get_value(self, device_name)
         return self._migrate_data(value) if self._migrate_data is not None else value
-        
+
+    def get_value_json(self, device_name: Optional[str] = None):
+        # Optionally override to change serialization behavior.
+        value = self.get_value(device_name)
+
+        if self._get_json_override_fn:
+            value = self._get_json_override_fn(value)
+
+        return value
+
+    def update_value_from_config(self, new_value):
+        # Optionally override to change deserialization behavior.
+        if self._config_set_override_fn:
+            new_value = self._config_set_override_fn(new_value)
+
+        self.update_value(new_value)
+
+    def override_value_from_config(self, device_name: str, override):
+        # Optionally override to change deserialization behavior.
+        if self._config_set_override_fn:
+            override = self._config_set_override_fn(override)
+
+        self.override_value(device_name, override)
 
     def on_device_name_changed(self, old_name: str, new_name: str):
         if old_name in self._overrides.keys():
@@ -177,8 +210,8 @@ class Setting(QtCore.QObject):
         self._override_widgets = {}
 
     def _create_widgets(
-            self, override_device_name: typing.Optional[str] = None) \
-            -> typing.Union[QtWidgets.QWidget, QtWidgets.QLayout]:
+            self, override_device_name: Optional[str] = None) \
+            -> Union[QtWidgets.QWidget, QtWidgets.QLayout]:
         '''
         Create the widgets necessary to manipulate this Setting in the UI.
 
@@ -198,7 +231,7 @@ class Setting(QtCore.QObject):
 
     def _on_widget_value_changed(
             self, new_value,
-            override_device_name: typing.Optional[str] = None):
+            override_device_name: Optional[str] = None):
         '''
         Update this Setting in response to a change in value caused by UI
         manipulation.
@@ -228,7 +261,7 @@ class Setting(QtCore.QObject):
 
     def _on_setting_changed(
             self, new_value,
-            override_device_name: typing.Optional[str] = None):
+            override_device_name: Optional[str] = None):
         '''
         Callback invoked when the value of this Setting changes.
 
@@ -241,9 +274,9 @@ class Setting(QtCore.QObject):
         pass
 
     def create_ui(
-            self, override_device_name: typing.Optional[str] = None,
-            form_layout: typing.Optional[QtWidgets.QFormLayout] = None) \
-            -> typing.Union[QtWidgets.QWidget, QtWidgets.QLayout, None]:
+            self, override_device_name: Optional[str] = None,
+            form_layout: Optional[QtWidgets.QFormLayout] = None) \
+            -> Union[QtWidgets.QWidget, QtWidgets.QLayout, None]:
         '''
         Create the UI for this Setting.
 
@@ -282,28 +315,28 @@ class Setting(QtCore.QObject):
             )
 
         return top_level_widget
-    
+
     def _register_on_setting_changed(self, top_level_widget: QtWidgets.QWidget, override_device_name: str):
         on_setting_changed_lambda = lambda old_value, new_value, override_device_name=override_device_name: \
             self._on_setting_changed(new_value, override_device_name=override_device_name)
         self.signal_setting_changed.connect(
             on_setting_changed_lambda
         )
-        
+
         # Clear the widget when it is destroyed to avoid dangling references
         top_level_widget.destroyed.connect(lambda destroyed_object=None:
             self._on_widget_destroyed(on_setting_changed_lambda, override_device_name)
         )
-        
+
     def _on_widget_destroyed(self, on_setting_changed_lambda, override_device_name: str):
         self.signal_setting_changed.disconnect(on_setting_changed_lambda)
         self.set_widget(widget=None, override_device_name=override_device_name)
-    
+
     def _decorate_with_reset_widget(self, override_device_name: str, setting_editor_widget: QtWidgets.QWidget):
         # Reset will still be shown on overrides
         if not self._allow_reset and override_device_name is None:
             return setting_editor_widget
-        
+
         horizontal_box = QtWidgets.QWidget()
         horizontal_layout = QtWidgets.QHBoxLayout(horizontal_box)
         horizontal_layout.setContentsMargins(0, 0, 0, 0)
@@ -327,7 +360,7 @@ class Setting(QtCore.QObject):
             lambda override_device_name=override_device_name:
                 self._on_press_reset_override(override_device_name)
         )
-        
+
         horizontal_layout.addWidget(button)
         is_base_reset_widget = override_device_name is None
         if is_base_reset_widget:
@@ -337,21 +370,21 @@ class Setting(QtCore.QObject):
             self._reset_override_widgets[override_device_name] = button
             self._refresh_reset_override_widget(override_device_name)
         return horizontal_box
-    
+
     def _on_press_reset_override(self, override_device_name: str):
         if override_device_name is None:
             self.update_value(self._original_value)
             self._base_reset_widget.setVisible(False)
             self._refresh_reset_override_widgets()
             return
-        
+
         if self.is_overridden(override_device_name):
             self.override_value(override_device_name, self._value)
             # Update UI
             self._on_widget_value_changed(self._value, override_device_name)
             self._on_setting_changed(self._value, override_device_name=override_device_name)
             self._reset_override_widgets[override_device_name].setVisible(False)
-            
+
     def _refresh_reset_base_widget(self):
         if self._base_reset_widget:
             self._base_reset_widget.setVisible(self._value != self._original_value)
@@ -365,8 +398,8 @@ class Setting(QtCore.QObject):
             self._reset_override_widgets[device_name].setVisible(self.is_overridden(device_name))
 
     def set_widget(
-            self, widget: typing.Optional[QtWidgets.QWidget] = None,
-            override_device_name: typing.Optional[str] = None):
+            self, widget: Optional[QtWidgets.QWidget] = None,
+            override_device_name: Optional[str] = None):
         '''
         Set the widget to be used to manipulate this Setting, or
         this particular device's override of the Setting.
@@ -387,8 +420,8 @@ class Setting(QtCore.QObject):
                 self._override_widgets[override_device_name] = widget
 
     def get_widget(
-            self, override_device_name: typing.Optional[str] = None) \
-            -> typing.Optional[QtWidgets.QWidget]:
+            self, override_device_name: Optional[str] = None) \
+            -> Optional[QtWidgets.QWidget]:
         '''
         Get the widget to be used to manipulate this Setting, or
         this particular device's override of the Setting.
@@ -397,7 +430,7 @@ class Setting(QtCore.QObject):
         '''
         return self._override_widgets.get(
             override_device_name, self._base_widget)
-    
+
 
 class BoolSetting(Setting):
     '''
@@ -405,7 +438,7 @@ class BoolSetting(Setting):
     '''
 
     def _create_widgets(
-            self, override_device_name: typing.Optional[str] = None) \
+            self, override_device_name: Optional[str] = None) \
             -> QtWidgets.QCheckBox:
         check_box = QtWidgets.QCheckBox()
         check_box.setChecked(self.get_value(override_device_name))
@@ -424,7 +457,7 @@ class BoolSetting(Setting):
 
     def _on_setting_changed(
             self, new_value: bool,
-            override_device_name: typing.Optional[str] = None):
+            override_device_name: Optional[str] = None):
         widget = self.get_widget(override_device_name=override_device_name)
         if widget:
             widget.setChecked(new_value)
@@ -442,16 +475,16 @@ class IntSetting(Setting):
     '''
     A UI-displayable Setting for storing and modifying an integer value.
     '''
-    
+
     def __init__(
         self,
         attr_name: str,
         nice_name: str,
         value: str,
-        tool_tip: typing.Optional[str] = None,
+        tool_tip: Optional[str] = None,
         show_ui: bool = True,
         allow_reset: bool = True,
-        migrate_data: typing.Callable[[any], None] = None,
+        migrate_data: Optional[Callable[[Any], None]] = None,
         is_read_only: bool = False
     ):
         '''
@@ -472,7 +505,7 @@ class IntSetting(Setting):
         self.is_read_only = is_read_only
 
     def _create_widgets(
-            self, override_device_name: typing.Optional[str] = None) \
+            self, override_device_name: Optional[str] = None) \
             -> QtWidgets.QLineEdit:
         line_edit = QtWidgets.QLineEdit()
         if self.tool_tip:
@@ -498,7 +531,7 @@ class IntSetting(Setting):
 
     def _on_setting_changed(
             self, new_value: int,
-            override_device_name: typing.Optional[str] = None):
+            override_device_name: Optional[str] = None):
         widget = self.get_widget(override_device_name=override_device_name)
         if not widget:
             return
@@ -520,10 +553,10 @@ class StringSetting(Setting):
         nice_name: str,
         value: str,
         placeholder_text: str = '',
-        tool_tip: typing.Optional[str] = None,
+        tool_tip: Optional[str] = None,
         show_ui: bool = True,
         allow_reset: bool = True,
-        migrate_data: typing.Callable[[any], None] = None,
+        migrate_data: Optional[Callable[[Any], None]] = None,
         is_read_only: bool = False
     ):
         '''
@@ -546,7 +579,7 @@ class StringSetting(Setting):
         self.is_read_only = is_read_only
 
     def _create_widgets(
-            self, override_device_name: typing.Optional[str] = None) \
+            self, override_device_name: Optional[str] = None) \
             -> QtWidgets.QLineEdit:
         line_edit = QtWidgets.QLineEdit()
         if self.tool_tip:
@@ -572,7 +605,7 @@ class StringSetting(Setting):
 
     def _on_setting_changed(
             self, new_value: str,
-            override_device_name: typing.Optional[str] = None):
+            override_device_name: Optional[str] = None):
         widget = self.get_widget(override_device_name=override_device_name)
         if not widget:
             return
@@ -593,7 +626,7 @@ class FileSystemPathSetting(StringSetting):
     '''
 
     def _getFileSystemPath(
-            self, parent: QtWidgets.QWidget = None,
+            self, parent: Optional[QtWidgets.QWidget] = None,
             start_path: str = '') -> str:
         raise NotImplementedError(
             f'Setting "{self.nice_name}" uses the FileSystemPathSetting '
@@ -601,7 +634,7 @@ class FileSystemPathSetting(StringSetting):
             'FilePathSetting) must be used instead.')
 
     def _create_widgets(
-            self, override_device_name: typing.Optional[str] = None) \
+            self, override_device_name: Optional[str] = None) \
             -> QtWidgets.QHBoxLayout:
         line_edit = super()._create_widgets(
             override_device_name=override_device_name)
@@ -643,7 +676,7 @@ class DirectoryPathSetting(FileSystemPathSetting):
     '''
 
     def _getFileSystemPath(
-            self, parent: QtWidgets.QWidget = None,
+            self, parent: Optional[QtWidgets.QWidget] = None,
             start_path: str = '') -> str:
         return QtWidgets.QFileDialog.getExistingDirectory(
             parent=parent, dir=start_path)
@@ -662,7 +695,7 @@ class FilePathSetting(FileSystemPathSetting):
         value: str,
         placeholder_text: str = '',
         file_path_filter: str = '',
-        tool_tip: typing.Optional[str] = None,
+        tool_tip: Optional[str] = None,
         show_ui: bool = True
     ):
         '''
@@ -684,7 +717,7 @@ class FilePathSetting(FileSystemPathSetting):
         self.file_path_filter = file_path_filter
 
     def _getFileSystemPath(
-            self, parent: QtWidgets.QWidget = None,
+            self, parent: Optional[QtWidgets.QWidget] = None,
             start_path: str = '') -> str:
         file_path, _ = QtWidgets.QFileDialog.getOpenFileName(
             parent=parent, dir=start_path, filter=self.file_path_filter)
@@ -697,7 +730,7 @@ class PerforcePathSetting(StringSetting):
     represents a Perforce depot path.
     '''
 
-    def _filter_value(self, value: typing.Optional[str]) -> str:
+    def _filter_value(self, value: Optional[str]) -> str:
         '''
         Clean the p4 path value by removing whitespace and trailing '/'.
         '''
@@ -718,11 +751,11 @@ class OptionSetting(Setting):
         attr_name: str,
         nice_name: str,
         value,
-        possible_values: typing.List = None,
-        tool_tip: typing.Optional[str] = None,
+        possible_values: Optional[list] = None,
+        tool_tip: Optional[str] = None,
         show_ui: bool = True,
         allow_reset: bool = True,
-        migrate_data: typing.Callable[[any], None] = None
+        migrate_data: Optional[Callable[[Any], None]] = None,
     ):
         '''
         Create a new OptionSetting object.
@@ -737,12 +770,13 @@ class OptionSetting(Setting):
         '''
         super().__init__(
             attr_name, nice_name, value,
-            tool_tip=tool_tip, show_ui=show_ui, allow_reset=allow_reset, migrate_data=migrate_data)
+            tool_tip=tool_tip, show_ui=show_ui, allow_reset=allow_reset,
+            migrate_data=migrate_data)
 
         self.possible_values = possible_values or []
 
     def _create_widgets(
-        self, override_device_name: typing.Optional[str] = None, *,
+        self, override_device_name: Optional[str] = None, *,
         widget_class: Type[NonScrollableComboBox] = NonScrollableComboBox
     ) -> NonScrollableComboBox:
         combo = widget_class()
@@ -767,7 +801,7 @@ class OptionSetting(Setting):
 
     def _on_setting_changed(
             self, new_value,
-            override_device_name: typing.Optional[str] = None):
+            override_device_name: Optional[str] = None):
         widget = self.get_widget(override_device_name=override_device_name)
         if not widget:
             return
@@ -787,8 +821,8 @@ class MultiOptionSetting(OptionSetting):
     '''
 
     def _create_widgets(
-            self, override_device_name: typing.Optional[str] = None) \
-            -> sb_widgets.MultiSelectionComboBox:
+            self, override_device_name: Optional[str] = None
+    ) -> sb_widgets.MultiSelectionComboBox:
         combo = sb_widgets.MultiSelectionComboBox()
         if self.tool_tip:
             combo.setToolTip(self.tool_tip)
@@ -811,7 +845,7 @@ class MultiOptionSetting(OptionSetting):
 
     def _on_setting_changed(
             self, new_value,
-            override_device_name: typing.Optional[str] = None):
+            override_device_name: Optional[str] = None):
         widget = self.get_widget(override_device_name=override_device_name)
         if not widget:
             return
@@ -821,33 +855,33 @@ class MultiOptionSetting(OptionSetting):
         widget_items = [item for item in widget_items if item.isEnabled()]
 
         old_value = [
-            item.text()
+            item.data(QtCore.Qt.UserRole)
             for item in widget_items
             if item.checkState() == QtCore.Qt.Checked]
         if new_value != old_value:
-            selected_items = []
+            selected_item_texts = []
             for item in widget_items:
-                if item.text() in new_value:
+                if item.data(QtCore.Qt.UserRole) in new_value:
                     item.setCheckState(QtCore.Qt.Checked)
-                    selected_items.append(item.text())
+                    selected_item_texts.append(item.text())
                 else:
                     item.setCheckState(QtCore.Qt.Unchecked)
 
-            widget.setEditText(widget.separator.join(selected_items))
+            widget.setEditText(widget.separator.join(selected_item_texts))
 
 
 class ListRow(QtWidgets.QWidget):
     INSERT_TEXT = "Insert"
     DUPLICATE_TEXT = "Duplicate"
     DELETE_TEXT = "Delete"
-    
+
     def __init__(
         self,
         array_index: int,
         editor_widget: QtWidgets.QWidget,
-        insert_item_callback: typing.Callable[[], None] = None, 
-        duplicate_item_callback: typing.Callable[[], None] = None, 
-        delete_item_callback: typing.Callable[[], None] = None,
+        insert_item_callback: Optional[Callable[[], None]] = None,
+        duplicate_item_callback: Optional[Callable[[], None]] = None,
+        delete_item_callback: Optional[Callable[[], None]] = None,
         parent=None
     ):
         super().__init__(parent=parent)
@@ -855,16 +889,16 @@ class ListRow(QtWidgets.QWidget):
         self._insert_item_callback = insert_item_callback
         self._duplicate_item_callback = duplicate_item_callback
         self._delete_item_callback = delete_item_callback
-    
+
         layout = QtWidgets.QHBoxLayout(self)
         layout.setSpacing(1)
         layout.setContentsMargins(0, 0, 0, 0)
-        
+
         # Shift the elements to the right
         layout.addItem(
             QtWidgets.QSpacerItem(10, 0, QtWidgets.QSizePolicy.Maximum, QtWidgets.QSizePolicy.Maximum)
         )
-        
+
         self._index_label = QtWidgets.QLabel()
         self._index_label.setFixedWidth(20)
         self.update_index(array_index)
@@ -886,11 +920,11 @@ class ListRow(QtWidgets.QWidget):
             self._duplicate_item_callback()
         elif selected_item == self.DELETE_TEXT:
             self._delete_item_callback()
-    
+
     @property
     def editor_widget(self):
         return self._editor_widget
-    
+
     def update_index(self, index: int):
         self._index_label.setText(str(index))
 
@@ -899,7 +933,7 @@ class ListSetting(Setting):
     '''
     A setting which has add and clear buttons. Each item will be displayed in a new line.
     Functions like array properties in Unreal Editor.
-    
+
     Subclasses are responsible for generating widgets for the array contents, see e.g. ArrayStringSetting.
     '''
 
@@ -907,11 +941,11 @@ class ListSetting(Setting):
         self,
         attr_name: str,
         nice_name: str,
-        value: typing.List = [],
-        tool_tip: typing.Optional[str] = None,
+        value: list = [],
+        tool_tip: Optional[str] = None,
         show_ui: bool = True,
         allow_reset: bool = True,
-        migrate_data: typing.Callable[[any], None] = None
+        migrate_data: Optional[Callable[[Any], None]] = None
     ):
         '''
         Create a new ListSetting object.
@@ -926,33 +960,33 @@ class ListSetting(Setting):
         super().__init__(
             attr_name, nice_name, value,
             tool_tip=tool_tip, show_ui=show_ui, allow_reset=allow_reset, migrate_data=migrate_data)
-        
+
         self.array_count_labels = {}
         self.element_layouts = {}
-    
-    def create_element(self, override_device_name: str, index: int) -> typing.Tuple[QtWidgets.QWidget, object]:
+
+    def create_element(self, override_device_name: str, index: int) -> Tuple[QtWidgets.QWidget, object]:
         '''
         Called when a new array element is supposed to be created.
         @returns The widget to use for editing and the default value for the new array element
         '''
         raise NotImplementedError("Subclasses must override this")
-    
-    def update_element_value(self, editor_widget: QtWidgets.QWidget, list_value: typing.List, index: int):
+
+    def update_element_value(self, editor_widget: QtWidgets.QWidget, list_value: list, index: int):
         '''
         Called to update the editor_widget with the the value from list_value[index].
         '''
         raise NotImplementedError("Subclasses must override this")
 
-    def _create_widgets(self, override_device_name: typing.Optional[str] = None):
+    def _create_widgets(self, override_device_name: Optional[str] = None):
         root = QtWidgets.QWidget()
         root_layout = QtWidgets.QVBoxLayout(root)
         root_layout.setSpacing(1)
         root_layout.setContentsMargins(1, 1, 1, 1)
-        
+
         root_layout.addWidget(
             self._create_header(override_device_name)
         )
-        
+
         elements_root = QtWidgets.QWidget()
         elements_layout = QtWidgets.QVBoxLayout(elements_root)
         self.element_layouts[override_device_name] = elements_layout
@@ -972,7 +1006,7 @@ class ListSetting(Setting):
         header_layout = QtWidgets.QHBoxLayout(header)
         header_layout.setSpacing(5)
         header_layout.setContentsMargins(0, 0, 0, 0)
-        
+
         array_count_label = QtWidgets.QLabel()
         self.array_count_labels[override_device_name] = array_count_label
         array_count_label.setFixedWidth(100)
@@ -1010,11 +1044,11 @@ class ListSetting(Setting):
         header_layout.addItem(QtWidgets.QSpacerItem(0, 0, QtWidgets.QSizePolicy.Expanding, QtWidgets.QSizePolicy.Minimum))
 
         return header
-        
+
     def _update_array_count_label(self, override_device_name: str):
         current_value = self.get_value(override_device_name)
         self.array_count_labels[override_device_name].setText(f"{len(current_value)} Elements")
-        
+
     def _on_press_add(self, override_device_name: str):
         current_value = self.get_value(override_device_name)
         row_widget, default_value = self._create_element_row(override_device_name, len(current_value))
@@ -1025,17 +1059,17 @@ class ListSetting(Setting):
         self._on_widget_value_changed(current_value + [default_value], override_device_name)
         # Force update in case _on_widget_value_changed implicitly changed array settings without telling us
         self._on_setting_changed(self.get_value(override_device_name), override_device_name)
-        
+
     def _on_press_clear(self, override_device_name: str):
         self._on_widget_value_changed([], override_device_name)
-        
+
         container_layout = self.element_layouts[override_device_name]
         for i in reversed(range(container_layout.count())):
             container_layout.itemAt(i).widget().setParent(None)
-        
+
         # Force update in case _on_widget_value_changed implicitly changed array settings without telling us
         self._on_setting_changed(self.get_value(override_device_name), override_device_name)
-    
+
     def _insert_at(self, override_device_name: str, index: int):
         current_value = self.get_value(override_device_name)
         row_widget, default_value = self._create_element_row(override_device_name, len(current_value))
@@ -1045,11 +1079,11 @@ class ListSetting(Setting):
 
         copied_value = current_value.copy()
         copied_value.insert(index, default_value)
-        
+
         self._on_widget_value_changed(copied_value, override_device_name)
         # Force update in case _on_widget_value_changed implicitly changed array settings without telling us
         self._on_setting_changed(self.get_value(override_device_name), override_device_name)
-        
+
     def _duplicate_at(self, override_device_name: str, index: int):
         current_value = self.get_value(override_device_name)
         row_widget, _ = self._create_element_row(override_device_name, len(current_value))
@@ -1063,7 +1097,7 @@ class ListSetting(Setting):
         self._on_widget_value_changed(copied_value, override_device_name)
         # Force update in case _on_widget_value_changed implicitly changed array settings without telling us
         self._on_setting_changed(self.get_value(override_device_name), override_device_name)
-        
+
     def _remove_at(self, override_device_name: str, index: int):
         current_value = self.get_value(override_device_name)
         copied_value = current_value.copy()
@@ -1073,7 +1107,7 @@ class ListSetting(Setting):
         # Force update in case _on_widget_value_changed implicitly changed array settings without telling us
         self._on_setting_changed(self.get_value(override_device_name), override_device_name)
 
-    def _create_element_row(self, override_device_name: str, index: int) -> typing.Tuple[ListRow, object]:
+    def _create_element_row(self, override_device_name: str, index: int) -> Tuple[ListRow, object]:
         editing_widget, default_value = self.create_element(override_device_name, index)
         row = ListRow(
             index,
@@ -1087,9 +1121,9 @@ class ListSetting(Setting):
         )
         return row, default_value
 
-    def _on_setting_changed(self, new_value: typing.List, override_device_name: typing.Optional[str] = None):
+    def _on_setting_changed(self, new_value: list, override_device_name: Optional[str] = None):
         container_layout = self.element_layouts[override_device_name]
-        
+
         container_len = container_layout.count()
         new_len = len(new_value)
         new_list_is_bigger = container_layout.count() < new_len
@@ -1098,18 +1132,18 @@ class ListSetting(Setting):
             for missing_index in range(container_layout.count(), len(new_value), 1):
                 new_array_row, _ = self._create_element_row(override_device_name, missing_index)
                 container_layout.addWidget(new_array_row)
-        
+
         if new_list_is_smaller:
             for added_index in reversed(range(new_len, container_layout.count(), 1)):
                 container_layout.itemAt(added_index).widget().deleteLater()
-                
+
         for index in range(new_len):
             array_row: ListRow = container_layout.itemAt(index).widget()
             self.update_element_value(array_row.editor_widget, new_value, index)
 
         self._update_array_count_label(override_device_name)
-    
-    
+
+
 class StringListSetting(ListSetting):
     '''
     An array setting where the elements are strings
@@ -1118,11 +1152,11 @@ class StringListSetting(ListSetting):
         self,
         attr_name: str,
         nice_name: str,
-        value: typing.List[str] = [],
-        tool_tip: typing.Optional[str] = None,
+        value: list[str] = [],
+        tool_tip: Optional[str] = None,
         show_ui: bool = True,
         allow_reset: bool = True,
-        migrate_data: typing.Callable[[any], None] = None
+        migrate_data: Optional[Callable[[Any], None]] = None
     ):
         '''
         Create a new ArraySetting object.
@@ -1139,7 +1173,7 @@ class StringListSetting(ListSetting):
             tool_tip=tool_tip, show_ui=show_ui, allow_reset=allow_reset, migrate_data=migrate_data)
         pass
 
-    def create_element(self, override_device_name: str, index: int) -> typing.Tuple[QtWidgets.QWidget, object]:
+    def create_element(self, override_device_name: str, index: int) -> Tuple[QtWidgets.QWidget, object]:
         line_edit = QtWidgets.QLineEdit()
         line_edit.editingFinished.connect(
             lambda line_edit=line_edit, override_device_name=override_device_name, index=index:
@@ -1154,7 +1188,7 @@ class StringListSetting(ListSetting):
         copied_value[index] = line_edit.text()
         self._on_widget_value_changed(copied_value, override_device_name)
 
-    def update_element_value(self, editor_widget: QtWidgets.QLineEdit, list_value: typing.List, index: int):
+    def update_element_value(self, editor_widget: QtWidgets.QLineEdit, list_value: list, index: int):
         editor_widget.setText(list_value[index])
 
 
@@ -1179,9 +1213,9 @@ class LoggingModel(QtGui.QStandardItemModel):
 
     def __init__(
             self,
-            categories: typing.List[str],
-            verbosity_levels: typing.List[str],
-            parent: typing.Optional[QtCore.QObject] = None):
+            categories: list[str],
+            verbosity_levels: list[str],
+            parent: Optional[QtCore.QObject] = None):
         '''
         Create a new LoggingModel object.
 
@@ -1208,15 +1242,15 @@ class LoggingModel(QtGui.QStandardItemModel):
                 'category')
 
     @property
-    def categories(self) -> typing.List[str]:
+    def categories(self) -> list[str]:
         return self._categories
 
     @property
-    def user_categories(self) -> typing.List[str]:
+    def user_categories(self) -> list[str]:
         return self._user_categories
 
     @property
-    def verbosity_levels(self) -> typing.List[str]:
+    def verbosity_levels(self) -> list[str]:
         return self._verbosity_levels
 
     def is_user_category(self, category: str) -> bool:
@@ -1296,7 +1330,7 @@ class LoggingModel(QtGui.QStandardItemModel):
     @category_verbosities.setter
     def category_verbosities(
             self,
-            value: typing.Optional[typing.Dict[str, typing.Optional[str]]]):
+            value: Optional[dict[str, Optional[str]]]):
         '''
         Sets the data in the model using the given dictionary of category
         names to verbosity levels (or None).
@@ -1349,7 +1383,7 @@ class LoggingVerbosityItemDelegate(QtWidgets.QStyledItemDelegate):
 
     def __init__(
             self,
-            verbosity_levels: typing.List[str],
+            verbosity_levels: list[str],
             parent: QtWidgets.QTreeView):
         super().__init__(parent=parent)
         self._verbosity_levels = verbosity_levels or []
@@ -1412,7 +1446,7 @@ class LoggingSettingVerbosityView(QtWidgets.QTreeView):
     def __init__(
             self,
             logging_model: LoggingModel,
-            parent: typing.Optional[QtWidgets.QWidget] = None):
+            parent: Optional[QtWidgets.QWidget] = None):
         super().__init__(parent=parent)
 
         self.setModel(logging_model)
@@ -1448,7 +1482,7 @@ class LoggingSettingVerbosityView(QtWidgets.QTreeView):
             verbosity_level_index = self.model().index(row, LoggingModel.VERBOSITY_COLUMN)
             self.openPersistentEditor(verbosity_level_index)
         
-    def selected_categories(self) -> typing.List[str]:
+    def selected_categories(self) -> list[str]:
         model = self.model()
         selected_categories = []
         for selection_range in self.selectionModel().selection():
@@ -1460,7 +1494,7 @@ class LoggingSettingVerbosityView(QtWidgets.QTreeView):
                 
         return selected_categories
         
-    def update_category_verbosities(self, categories: typing.List[str], new_verbosity: str):
+    def update_category_verbosities(self, categories: list[str], new_verbosity: str):
         model = self.model()
         category_verbosities = model.category_verbosities
         for category in categories:
@@ -1488,8 +1522,8 @@ class LoggingSetting(Setting):
 
     def _filter_value(
             self,
-            value: typing.Optional[typing.Dict[str, typing.Optional[str]]]) \
-            -> collections.OrderedDict:
+            value: Optional[dict[str, Optional[str]]]
+    )-> collections.OrderedDict:
         '''
         Filter function to modify the incoming value before updating or
         overriding the setting.
@@ -1522,13 +1556,13 @@ class LoggingSetting(Setting):
         self,
         attr_name: str,
         nice_name: str,
-        value: typing.Dict[str, typing.Optional[str]],
-        categories: typing.List[str] = None,
-        verbosity_levels: typing.List[str] = None,
-        tool_tip: typing.Optional[str] = None,
+        value: dict[str, Optional[str]],
+        categories: Optional[list[str]] = None,
+        verbosity_levels: Optional[list[str]] = None,
+        tool_tip: Optional[str] = None,
         show_ui: bool = True,
         allow_reset: bool = True,
-        migrate_data: typing.Callable[[any], None] = None
+        migrate_data: Optional[Callable[[Any], None]] = None
     ):
         '''
         Create a new LoggingSetting object.
@@ -1556,8 +1590,8 @@ class LoggingSetting(Setting):
             verbosity_levels or self.DEFAULT_VERBOSITY_LEVELS)
 
     def get_command_line_arg(
-            self, override_device_name: typing.Optional[str] = None) \
-            -> str:
+            self, override_device_name: Optional[str] = None
+    ) -> str:
         '''
         Generate the command line argument for specifying the logging
         configuration based on the value currently stored in the Setting.
@@ -1577,8 +1611,8 @@ class LoggingSetting(Setting):
         return f'-LogCmds=\"{", ".join(logging_strings)}\"'
 
     def _create_widgets(
-            self, override_device_name: typing.Optional[str] = None) \
-            -> QtWidgets.QVBoxLayout:
+            self, override_device_name: Optional[str] = None
+    ) -> QtWidgets.QVBoxLayout:
         model = LoggingModel(
             categories=self._categories,
             verbosity_levels=self._verbosity_levels)
@@ -1702,7 +1736,7 @@ class LoggingSetting(Setting):
 
     def _on_setting_changed(
             self, new_value,
-            override_device_name: typing.Optional[str] = None):
+            override_device_name: Optional[str] = None):
         widget = self.get_widget(override_device_name=override_device_name)
         if not widget:
             return
@@ -1734,7 +1768,7 @@ class AddressSetting(OptionSetting):
             migrate_data=migrate_data)
 
     def _create_widgets(
-        self, override_device_name: typing.Optional[str] = None
+        self, override_device_name: Optional[str] = None
     ) -> NonScrollableComboBox:
         combo: NonScrollableComboBox = super()._create_widgets(
             override_device_name, widget_class=sb_widgets.AddressComboBox)
@@ -1806,7 +1840,7 @@ class ConfigPathIsUserSettingsError(ConfigPathError):
 
 
 def get_absolute_config_path(
-        config_path: typing.Union[str, pathlib.Path]) -> pathlib.Path:
+        config_path: Union[str, pathlib.Path]) -> pathlib.Path:
     '''
     Returns the given string or path object as an absolute config path.
 
@@ -1842,7 +1876,7 @@ def get_absolute_config_path(
 
 
 def get_relative_config_path(
-        config_path: typing.Union[str, pathlib.Path]) -> pathlib.Path:
+        config_path: Union[str, pathlib.Path]) -> pathlib.Path:
     '''
     Returns the given string or path object as a config path relative to the
     root configs path.
@@ -1901,10 +1935,10 @@ class Config(object):
         '''
         self.saving_allowed = self.saving_allowed_fifo.pop()
         
-    def init(self, file_path: typing.Union[str, pathlib.Path]):
+    def init(self, file_path: Union[str, pathlib.Path]):
         self.init_with_file_path(file_path)
 
-    def init_with_file_path(self, file_path: typing.Union[str, pathlib.Path]):
+    def init_with_file_path(self, file_path: Union[str, pathlib.Path]):
         if file_path:
             try:
                 self.file_path = get_absolute_config_path(file_path)
@@ -2000,7 +2034,7 @@ class Config(object):
                 os.remove(new_file_path)
             shutil.copy(original_file_path, new_file_path)
 
-    def init_new_config(self, file_path: typing.Union[str, pathlib.Path], uproject, engine_dir, p4_settings):
+    def init_new_config(self, file_path: Union[str, pathlib.Path], uproject, engine_dir, p4_settings):
         ''' 
         Initialize new configuration
         '''
@@ -2323,7 +2357,8 @@ class Config(object):
         if loaded_settings:
             for setting in settings:
                 if setting.attr_name in loaded_settings:
-                    setting.update_value(loaded_settings[setting.attr_name])
+                    value = loaded_settings[setting.attr_name]
+                    setting.update_value_from_config(value)
             del self._plugin_data_from_config[device_type]
 
     def register_plugin_settings(self, device_type, settings):
@@ -2352,7 +2387,7 @@ class Config(object):
         if device_name in known_devices:
             self.save()
 
-    def replace(self, new_config_path: typing.Union[str, pathlib.Path]):
+    def replace(self, new_config_path: Union[str, pathlib.Path]):
         """
         Move the file.
 
@@ -2367,7 +2402,7 @@ class Config(object):
         self.file_path = new_config_path
         self.save()
 
-    def save_as(self, new_config_path: typing.Union[str, pathlib.Path]):
+    def save_as(self, new_config_path: Union[str, pathlib.Path]):
         """
         Copy the file.
 
@@ -2438,7 +2473,7 @@ class Config(object):
             settings = {}
 
             for setting in plugin_settings:
-                settings[setting.attr_name] = setting.get_value()
+                settings[setting.attr_name] = setting.get_value_json()
 
             data["devices"][device_type] = {
                 "settings": settings,
@@ -2454,21 +2489,32 @@ class Config(object):
             serialized_settings = {}
 
             for setting in settings:
-                serialized_settings[setting.attr_name] = setting.get_value()
+                value = setting.get_value_json()
+                serialized_settings[setting.attr_name] = value
 
             for setting in overrides:
                 if setting.is_overridden(device_name):
-                    serialized_settings[setting.attr_name] = setting.get_value(
-                        device_name)
+                    value = setting.get_value_json(device_name)
+                    serialized_settings[setting.attr_name] = value
 
             data["devices"][device_type][device_name] = serialized_settings
 
         # Save to file
-        #
-        self.file_path.parent.mkdir(parents=True, exist_ok=True)
-        with open(self.file_path, 'w') as f:
-            json.dump(data, f, indent=4)
-            LOGGER.debug(f'Config File: {self.file_path} updated')
+        temp_buf = io.StringIO()
+        try:
+            json.dump(data, temp_buf, indent=4)
+        except Exception as exc:
+            LOGGER.error('Error while serializing config', exc_info=exc)
+            return
+
+        try:
+            self.file_path.parent.mkdir(parents=True, exist_ok=True)
+            with open(self.file_path, 'w') as f:
+                temp_buf.seek(0)
+                shutil.copyfileobj(temp_buf, f)
+                LOGGER.debug(f'Config File: {self.file_path} updated')
+        except Exception as exc:
+            LOGGER.error('Error writing config to disk', exc_info=exc)
 
     def on_device_name_changed(self, old_name, new_name):
         old_key = None
@@ -2515,7 +2561,7 @@ class Config(object):
         return os.path.join(self.get_project_dir(), 'Content')
 
     def get_unreal_content_plugins(
-            self) -> typing.List[ue_plugin_utils.UnrealPlugin]:
+            self) -> list[ue_plugin_utils.UnrealPlugin]:
         '''
         Get a list of Unreal Engine plugins that match the current
         Switchboard config's content plugin filter settings.
@@ -2768,7 +2814,7 @@ class UserSettings(object):
             json.dump(data, f, indent=4)
 
 
-def list_config_paths() -> typing.List[pathlib.Path]:
+def list_config_paths() -> list[pathlib.Path]:
     '''
     Returns a list of absolute paths to all config files in the configs dir.
     '''
diff --git a/Engine/Plugins/VirtualProduction/Switchboard/Source/Switchboard/switchboard/devices/device_base.py b/Engine/Plugins/VirtualProduction/Switchboard/Source/Switchboard/switchboard/devices/device_base.py
--- a/Engine/Plugins/VirtualProduction/Switchboard/Source/Switchboard/switchboard/devices/device_base.py
+++ b/Engine/Plugins/VirtualProduction/Switchboard/Source/Switchboard/switchboard/devices/device_base.py
@@ -89,7 +89,7 @@ class Device(QtCore.QObject):
         for setting in self.setting_overrides():
             if setting.attr_name in kwargs.keys():
                 override = kwargs[setting.attr_name]
-                setting.override_value(self.name, override)
+                setting.override_value_from_config(self.name, override)
 
         self._project_changelist: Optional[int] = None
         self._engine_changelist: Optional[int] = None
diff --git a/Engine/Plugins/VirtualProduction/Switchboard/Source/Switchboard/switchboard/devices/unreal/plugin_unreal.py b/Engine/Plugins/VirtualProduction/Switchboard/Source/Switchboard/switchboard/devices/unreal/plugin_unreal.py
--- a/Engine/Plugins/VirtualProduction/Switchboard/Source/Switchboard/switchboard/devices/unreal/plugin_unreal.py
+++ b/Engine/Plugins/VirtualProduction/Switchboard/Source/Switchboard/switchboard/devices/unreal/plugin_unreal.py
@@ -3,7 +3,12 @@
 from __future__ import annotations
 
 import base64
+from collections import OrderedDict
 import concurrent.futures
+from datetime import datetime
+from functools import wraps
+import hashlib
+from ipaddress import IPv4Address
 import json
 import os
 import pathlib
@@ -11,19 +16,16 @@ import re
 import socket
 import sys
 import threading
+from typing import Callable, Optional, Union
 import uuid
-from collections import OrderedDict
-from datetime import datetime
-from functools import wraps
-from ipaddress import IPv4Address
-from typing import Callable, List, Optional, Set
 
-import switchboard.switchboard_widgets as sb_widgets
 from PySide2 import QtCore, QtGui, QtWidgets
-from switchboard import config_osc as osc
-from switchboard import message_protocol, switchboard_application
-from switchboard import switchboard_dialog as sb_dialog
-from switchboard import switchboard_utils as sb_utils
+
+from switchboard import message_protocol, switchboard_application, ugs_utils
+import switchboard.config_osc as osc
+import switchboard.switchboard_dialog as sb_dialog
+import switchboard.switchboard_utils as sb_utils
+import switchboard.switchboard_widgets as sb_widgets
 from switchboard.config import (CONFIG, DEFAULT_MAP_TEXT, ENABLE_UGS_SUPPORT, SETTINGS, BoolSetting, DirectoryPathSetting,
                                 EngineSyncMethod, FilePathSetting, IntSetting, MultiOptionSetting, OptionSetting, Setting,
                                 StringListSetting, StringSetting, migrate_comma_separated_string_to_list)
@@ -33,9 +35,7 @@ from switchboard.devices.unreal.uassetparser import UassetParser
 from switchboard.listener_client import ListenerClient
 from switchboard.switchboard_logging import LOGGER
 from switchboard.tools.insights_launcher import InsightsLauncher
-
-from ...util import p4_changelist_inspection
-from ...util.p4_changelist_inspection import P4Error
+from switchboard.util import p4_changelist_inspection
 from . import version_helpers
 from .listener_watcher import ListenerWatcher
 from .redeploy_dialog import RedeployListenerDialog
@@ -97,7 +97,7 @@ class ProgramStartQueue:
         # and main thread.
         self.lock = threading.Lock()
 
-        self.queued_programs: List[ProgramStartQueueItem] = []
+        self.queued_programs: list[ProgramStartQueueItem] = []
         self.starting_programs: OrderedDict[
             uuid.UUID, ProgramStartQueueItem] = OrderedDict()
         # initialized as returned by listener
@@ -265,6 +265,7 @@ class ProgramStartQueue:
     def clear_running_programs(self):
         self.running_programs.clear()
 
+
 class LiveLinkPresetSetting(Setting):
     ''' Container of the LiveLink Preset setting
     Its main widget is a combo box that makes available the list
@@ -423,6 +424,7 @@ class LiveLinkPresetSetting(Setting):
                     combo.setCurrentIndex(item_idx)
                     break
 
+
 class MediaProfileSetting(LiveLinkPresetSetting):
     ''' Container of the MediaProfile setting.
     '''
@@ -441,9 +443,30 @@ class MediaProfileSetting(LiveLinkPresetSetting):
         return True
 
 
+class SyncCategoryOption:
+    ''' Used to represent a UGS sync filter category within a Setting. '''
+    def __init__(self, id: Union[str, uuid.UUID], name: str):
+        if isinstance(id, str):
+            id = uuid.UUID(id)
+
+        self.id = id
+        self.name = name
+
+    def __str__(self) -> str:
+        ''' Determines the text displayed in the `MultiSelectionComboBox`. '''
+        return self.name
+
+    def __eq__(self, other: object):
+        ''' Two category options compare equal if their UUIDs are equal. '''
+        if isinstance(other, SyncCategoryOption):
+            return self.id == other.id
+        else:
+            return False
+
+
 class DeviceUnreal(Device):
 
-    NDISPLAY_CLASS_NAMES =  (
+    NDISPLAY_CLASS_NAMES = (
         'DisplayClusterBlueprint',
         '/Script/DisplayCluster.DisplayClusterBlueprint',
     )
@@ -546,8 +569,8 @@ class DeviceUnreal(Device):
             value=':0',
             tool_tip=(
                 'Local interface binding (-UDPMESSAGING_TRANSPORT_UNICAST) of '
-                'the form {address}:{port}. If {address} is omitted, the device address '
-                'address is used.'),
+                'the form {address}:{port}. If {address} is omitted, the '
+                'device address is used.'),
         ),
         'udpmessaging_extra_static_endpoints': StringSetting(
             attr_name='udpmessaging_extra_static_endpoints',
@@ -636,6 +659,44 @@ class DeviceUnreal(Device):
             value='',
             tool_tip=('Adds the selected Media Profile to the command line')
         ),
+        'use_sync_filters': BoolSetting(
+            attr_name='use_sync_filters',
+            nice_name='Use Sync Filters',
+            value=False,
+            tool_tip=(
+                'Controls whether UnrealGameSync filter categories or custom '
+                'views specified below are used during Perforce sync')
+        ),
+        'included_sync_categories': MultiOptionSetting(
+            attr_name='included_sync_categories',
+            nice_name='Included Sync Filter Categories',
+            value=[
+                SyncCategoryOption('6703e989-d912-451d-93ad-b48de748d282',
+                                   'Content'),
+                SyncCategoryOption('f44b2d25-cbc0-4a8f-b6b3-e4a8125533dd',
+                                   'Platform Support: Linux'),
+                SyncCategoryOption('5206ccee-9024-4e36-8b89-f5f5a7d288d2',
+                                   'Platform Support: Win64'),
+                SyncCategoryOption('cfec942a-bb90-4f0c-accf-238ecaad9430',
+                                   'Source Code'),
+            ],
+            tool_tip='UnrealGameSync filter categories to include during sync',
+        ).with_get_json_override_fn(
+            lambda val_list: [(str(x.id), x.name) for x in val_list]
+        ).with_config_set_override_fn(
+            lambda val_list: [
+                SyncCategoryOption(uuid.UUID(x[0]), x[1])
+                for x in val_list]
+        ),
+        'custom_sync_view': StringSetting(
+            attr_name='custom_sync_view',
+            nice_name='Custom Sync View',
+            value='',
+            placeholder_text='-/Samples/Games/...,-.../ExcludePlatform/...',
+            tool_tip=(
+                'Comma separated. Used to specify freeform Perforce-style '
+                'wildcards to be applied during sync.'),
+        ),
     }
 
     if ENABLE_UGS_SUPPORT:
@@ -653,6 +714,7 @@ class DeviceUnreal(Device):
 
     mu_server = switchboard_application.get_multi_user_server_instance()
     rsync_server = switchboard_application.RsyncServer()
+    ugs_config: Optional[ugs_utils.IniParser] = None
 
     # Monitors the local listener executable and notifies when the file is
     # changed.
@@ -660,12 +722,88 @@ class DeviceUnreal(Device):
 
     # Every DeviceUnreal (and derived class, e.g. DevicenDisplay) instance;
     # used for listener updates.
-    active_unreal_devices: Set[DeviceUnreal] = set()
+    active_unreal_devices: set[DeviceUnreal] = set()
 
     # Flag used to batch together multiple rapid calls to
     # `_queue_notify_redeploy`.
     _pending_notify_redeploy = False
 
+    class UgsConfigRunnable(QtCore.QRunnable):
+        ''' Performs a UGS config update on the thread pool. '''
+
+        class Signals(QtCore.QObject):
+            result = QtCore.Signal(ugs_utils.IniParser)
+
+        def __init__(self, engine_dir: Optional[str], proj_dir: Optional[str]):
+            super().__init__()
+            self.engine_dir = engine_dir
+            self.proj_dir = proj_dir
+            self.signals = self.Signals()
+
+        def run(self):
+            ugs_config = ugs_utils.parse_depot_ugs_configs(
+                self.engine_dir, self.proj_dir)
+            self.signals.result.emit(ugs_config)
+
+    _ugs_config_runnable: Optional[UgsConfigRunnable] = None
+    _pending_ugs_config_refresh = False
+
+    @classmethod
+    def p4_refresh_ugs_config(cls):
+        if cls.ugs_config is None:
+            # One-time initialization
+            def refresh():
+                cls.p4_refresh_ugs_config()
+
+            CONFIG.P4_ENABLED.signal_setting_changed.connect(refresh)
+            CONFIG.P4_ENGINE_PATH.signal_setting_changed.connect(refresh)
+            CONFIG.P4_PROJECT_PATH.signal_setting_changed.connect(refresh)
+
+        if not CONFIG.P4_ENABLED.get_value():
+            cls.ugs_config = ugs_utils.IniParser()
+            return
+
+        if cls._pending_ugs_config_refresh:
+            return
+        elif cls._ugs_config_runnable:
+            cls._pending_ugs_config_refresh = True
+            return
+
+        engine_dir = CONFIG.P4_ENGINE_PATH.get_value() or None
+        project_dir = CONFIG.P4_PROJECT_PATH.get_value() or None
+        cls._ugs_config_runnable = cls.UgsConfigRunnable(
+            engine_dir, project_dir)
+
+        cls._ugs_config_runnable.signals.result.connect(
+            cls._on_updated_ugs_config)
+        QtCore.QThreadPool.globalInstance().start(cls._ugs_config_runnable)
+
+    @classmethod
+    def _on_updated_ugs_config(cls, ugs_config: ugs_utils.IniParser):
+        cls.ugs_config = ugs_config
+
+        assert cls._ugs_config_runnable is not None
+        cls._ugs_config_runnable.signals.result.disconnect(
+            cls._on_updated_ugs_config)
+        cls._ugs_config_runnable = None
+
+        if cls._pending_ugs_config_refresh:
+            cls._pending_ugs_config_refresh = False
+            cls.p4_refresh_ugs_config()
+            return
+
+        sync_filters = ugs_utils.SyncFilters()
+        sync_filters.read_categories_from_ini_parser(cls.ugs_config)
+        sync_filters.categories = dict(sorted(
+            sync_filters.categories.items(), key=lambda item: item[1].name))
+
+        # The UUIDs are stable, but the names/object identities change
+        include_setting = cls.csettings['included_sync_categories']
+        include_setting.possible_values.clear()
+        for category in sync_filters.categories.values():
+            option = SyncCategoryOption(category.id, category.name)
+            include_setting.possible_values.append(option)
+
     @classmethod
     def get_designated_local_builder(cls) -> Optional[DeviceUnreal]:
         '''
@@ -675,7 +813,7 @@ class DeviceUnreal(Device):
         listener executables.
         '''
 
-        def ips_for_host(host: str) -> List[IPv4Address]:
+        def ips_for_host(host: str) -> list[IPv4Address]:
             try:
                 (primary_name, aliases, ipstrs) = socket.gethostbyname_ex(host)
                 return [IPv4Address(ipstr) for ipstr in ipstrs]
@@ -777,6 +915,13 @@ class DeviceUnreal(Device):
             show_ui=False
         )
 
+        self.last_sync_filter_hash = StringSetting(
+            attr_name="last_sync_filter_hash",
+            nice_name="Last UGS Sync Filter Hash",
+            value=kwargs.get('last_sync_filter_hash', ''),
+            show_ui=False
+        )
+
         self.exclude_from_build = BoolSetting(
             attr_name="exclude_from_build",
             nice_name="Exclude from build",
@@ -796,6 +941,7 @@ class DeviceUnreal(Device):
         self.runtime_str = ""
         self.inflight_project_cl = None
         self.inflight_engine_cl = None
+        self.inflight_sync_hash: Optional[str] = None
 
         listener_qt_handler = self.unreal_client.listener_qt_handler
         listener_qt_handler.listener_connecting.connect(
@@ -874,22 +1020,24 @@ class DeviceUnreal(Device):
                 'rsync_port'].signal_setting_changed.connect(
                     launch_rsync_server)
 
+        if self.ugs_config is None:
+            self.p4_refresh_ugs_config()
+
         app = QtCore.QCoreApplication.instance()
         app.aboutToQuit.connect(self._on_about_to_quit)
 
-        self._widget_classes: Set[str] = set()
+        self._widget_classes: set[str] = set()
 
         # Notify user of any invalid settings
         self.check_settings_valid()
 
     def init(self, widget_class, icons):
         super().init(widget_class, icons)
-        
+
         self.exclude_from_build.signal_setting_changed.connect(
             lambda _, new_value: self.on_setting_exclude_from_build_changed(new_value)
         )
         self.on_setting_exclude_from_build_changed(self.exclude_from_build.get_value())
-        
 
     def should_allow_exit(self, close_req_id: int) -> bool:
         # Delegate to a class method which surveys all active devices.
@@ -1002,7 +1150,8 @@ class DeviceUnreal(Device):
             self.last_launch_command,
             self.last_log_path,
             self.last_trace_path,
-            self.exclude_from_build
+            self.exclude_from_build,
+            self.last_sync_filter_hash,
         ]
 
     def check_settings_valid(self) -> bool:
@@ -1358,15 +1507,49 @@ class DeviceUnreal(Device):
         if generate_proj_files:
             sync_args += ' --generate'
 
+        if self.csettings['use_sync_filters'].get_value():
+            include_str = ''
+            custom_view_str = self.csettings['custom_sync_view'].get_value()
+
+            incl_categories: list[SyncCategoryOption] = self.csettings[
+                'included_sync_categories'].get_value()
+            if incl_categories:
+                incl_categories.sort(key=lambda cat: cat.id)
+                include_str = ','.join(str(cat.id) for cat in incl_categories)
+                sync_args += f' --include-categories="{include_str}"'
+
+            if custom_view_str:
+                sync_args += f' --custom-view="{custom_view_str}"'
+
+            # If the user changes their sync filters, we need to rescan their
+            # entire workspace to ensure we remove newly-excluded files.
+            # We make this determination by concatenating and hashing their
+            # selected filters, and then comparing that hash between syncs.
+            new_hash_input_bytes = bytearray()
+            new_hash_input_bytes.extend(include_str.encode('utf-8'))
+            new_hash_input_bytes.extend(custom_view_str.encode('utf-8'))
+            new_hash = hashlib.sha1(new_hash_input_bytes).hexdigest()
+            prev_hash = self.last_sync_filter_hash.get_value()
+            if new_hash != prev_hash:
+                LOGGER.info(
+                    f'Filter has changed ({prev_hash} -> {new_hash}); '
+                    'finding files in workspace that need to be removed.')
+                sync_args += ' --remove-excluded-workspace-files'
+
+                # If the sync is interrupted, we'll be in an indeterminate
+                # state, and we should always clean next time.
+                self.last_sync_filter_hash.update_value('INVALID')
+                self.inflight_sync_hash = new_hash  # persisted on sync success
+
         if ENABLE_UGS_SUPPORT:
             if sync_using_ugs:
                 sync_args += ' --use-ugs'
                 if self.unrealgamesync_lib_dir_setting:
                     sync_args += f' --ugs-lib-dir={self.unrealgamesync_lib_dir_setting}'
-        
+
             if sync_precompiled_bins:
                 sync_args += ' --use-pcbs'
-            
+
                 # If we're syncing 'Precompiled Binaries' one of those binaries may be the SwitchboardListener executable
                 # which means (on Windows atleast) we need to move the executable to make way for the new one
                 _, msg = message_protocol.create_free_listener_bin_message()
@@ -1657,8 +1840,8 @@ class DeviceUnreal(Device):
         else:
             return ''
 
-    def build_udpmessaging_static_endpoint_list(self) -> List[str]:
-        endpoints: List[str] = []
+    def build_udpmessaging_static_endpoint_list(self) -> list[str]:
+        endpoints: list[str] = []
 
         # Multi-user server.
         if CONFIG.MUSERVER_AUTO_ENDPOINT.get_value():
@@ -2023,6 +2206,8 @@ class DeviceUnreal(Device):
         elif program_name == 'sync':
             if returncode == 0:
                 LOGGER.info(f"{self.name}: Sync successful")
+                self.last_sync_filter_hash.update_value(
+                    self.inflight_sync_hash)
             else:
                 LOGGER.error(f"{self.name}: Sync failed!")
                 for line in get_stdout_str().splitlines():
@@ -2047,8 +2232,7 @@ class DeviceUnreal(Device):
             self.project_changelist = self.inflight_project_cl
             self.inflight_project_cl = None
 
-            # refresh project CL
-            self.project_changelist = self.project_changelist
+            self.inflight_sync_hash = None
 
             self.status = DeviceStatus.CLOSED
 
@@ -2159,12 +2343,13 @@ class DeviceUnreal(Device):
         self.built_engine_changelist = data.get("Changelist", None)
 
     def on_file_receive_failed(self, source_path, error):
-        roles = self.setting_roles.get_value()
-        if len(roles) > 0:
-            LOGGER.error(
-                f"{self.name}: Error receiving role file from listener and "
-                f"device claims to have these roles: {' | '.join(roles)}")
-            LOGGER.error(f"Error: {error}")
+        if source_path.endswith(DeviceUnreal.csettings["roles_filename"].get_value(self.name)):
+            roles = self.setting_roles.get_value()
+            if len(roles) > 0:
+                LOGGER.error(
+                    f"{self.name}: Error receiving role file from listener and "
+                    f"device claims to have these roles: {' | '.join(roles)}")
+                LOGGER.error(f"Error: {error}")
 
     def on_listener_programstdout(self, message):
         ''' Handles updates to stdout of programs
@@ -2976,7 +3161,7 @@ class DeviceWidgetUnreal(DeviceWidget):
                     later_cl,
                     CONFIG.P4_ENGINE_PATH.get_value()
                 )
-            except P4Error as error:
+            except p4_changelist_inspection.P4Error as error:
                 LOGGER.error(f"Couldn't check {built_cl} - {built_cl} for source code changes."
                              f" Reason: {error.message}")
                 # Assume that non-equal CL numbers have source changes
diff --git a/Engine/Plugins/VirtualProduction/Switchboard/Source/Switchboard/switchboard/p4_utils.py b/Engine/Plugins/VirtualProduction/Switchboard/Source/Switchboard/switchboard/p4_utils.py
--- a/Engine/Plugins/VirtualProduction/Switchboard/Source/Switchboard/switchboard/p4_utils.py
+++ b/Engine/Plugins/VirtualProduction/Switchboard/Source/Switchboard/switchboard/p4_utils.py
@@ -1,12 +1,20 @@
 # Copyright Epic Games, Inc. All Rights Reserved.
 
+from dataclasses import dataclass
+from functools import wraps
+import logging
 import marshal
 import pathlib
 import subprocess
-from functools import wraps
+from typing import Optional
 
 from . import switchboard_utils as sb_utils
-from .switchboard_logging import LOGGER
+from .devices.unreal.version_helpers import LISTENER_COMPATIBLE_VERSION
+
+
+# Metadata attached to p4 commands for sever analytics
+meta_zprog = 'switchboard'
+meta_zversion = '.'.join(str(x) for x in LISTENER_COMPATIBLE_VERSION)  # FIXME?
 
 
 def p4_login(f):
@@ -15,8 +23,8 @@ def p4_login(f):
         try:
             return f(*args, **kwargs)
         except Exception as e:
-            LOGGER.error(f'{repr(e)}')
-            LOGGER.error('Error running P4 command. Please make sure you are logged into Perforce and environment variables are set')
+            logging.error(f'{repr(e)}')
+            logging.error('Error running P4 command. Please make sure you are logged into Perforce and environment variables are set')
             return None
 
     return wrapped
@@ -28,7 +36,7 @@ def p4_latest_changelist(p4_path, working_dir, num_changelists=10):
     Return (num_changelists) latest CLs
     """
     p4_command = f'p4 -ztag -F "%change%" changes -m {num_changelists} {p4_path}/...'
-    LOGGER.info(f"Executing: {p4_command}")
+    logging.info(f"Executing: {p4_command}")
 
     p4_result = subprocess.check_output(p4_command, cwd=working_dir, shell=True, startupinfo=sb_utils.get_hidden_sp_startupinfo()).decode()
 
@@ -129,7 +137,7 @@ def p4_from_localpath(localpath, workspaces, preferredClient):
         client = next((valueForMarshalledKey(ws, 'client') for ws in wss if valueForMarshalledKey(ws, 'client')==preferredClient), preferredClient)
         if not client:
             client = valueForMarshalledKey(wss[0],'client')
-    
+
     # now we need to determine the corresponding p4 path to the full engine path
     wheres = run(cmd=f'-c {client} where', args=[localpath])
 
@@ -140,3 +148,250 @@ def p4_from_localpath(localpath, workspaces, preferredClient):
 
     return client, p4path
 
+
+##########
+
+
+def p4_variables() -> dict[str, str]:
+    '''
+    Returns Perforce variable (e.g. `P4CONFIG`, `p4 set`) key/value pairs.
+    '''
+    output = subprocess.check_output(['p4', 'set', '-q']).decode()
+    variables: dict[str, str] = {}
+    for line in output.splitlines():
+        var_name, _, var_value = line.partition('=')
+        variables[var_name] = var_value
+    return variables
+
+
+def p4(
+    cmd: str,
+    cmd_opts: list[str],
+    global_opts: Optional[list[str]] = None,
+    *,
+    user: Optional[str] = None,
+    client: Optional[str] = None,
+) -> subprocess.Popen:
+    ''' Runs a p4 command and includes some common/required switches. '''
+    args = ['p4', f'-zprog={meta_zprog}', f'-zversion={meta_zversion}',
+            '-ztag', '-G', '-Qutf8']
+
+    if user is not None:
+        args.extend(['-u', user])
+    if client is not None:
+        args.extend(['-c', client])
+
+    args.extend(global_opts or [])
+
+    args.append(cmd)
+    args.extend(cmd_opts)
+
+    logging.debug(f'p4(): invoking subprocess: args={args}')
+
+    # stdin=subprocess.DEVNULL is required when launched via
+    # SwitchboardListener; otherwise this call tries to make the non-existent
+    # stdin inheritable, raising `OSError: [WinError 6] The handle is invalid`
+    return subprocess.Popen(args, stdout=subprocess.PIPE,
+                            stdin=subprocess.DEVNULL)
+
+
+def p4_get_records(
+    cmd: str,
+    opts: list[str],
+    global_opts: Optional[list[str]] = None,
+    *,
+    include_info: bool = False,
+    include_error: bool = False,
+    log_unreturned: bool = True,
+    user: Optional[str] = None,
+    client: Optional[str] = None,
+) -> list[dict]:
+    results: list[dict] = []
+    proc = p4(cmd, opts, global_opts=global_opts, user=user, client=client)
+    while True:
+        try:
+            record = marshal.load(proc.stdout)
+        except EOFError:
+            break
+
+        # Possible codes: stat, info, error, text (only `p4 print`)
+        if (record[b'code'] == b'info') and not include_info:
+            if log_unreturned:
+                logging.debug(f'p4_get_records: info: {record}')
+            continue
+        elif (record[b'code'] == b'error') and not include_error:
+            if log_unreturned:
+                logging.debug(f'p4_get_records: error: {record}')
+            continue
+
+        results.append(record)
+
+    proc.stdout.close()
+    return results
+
+
+def p4_get_decoded_records(
+    cmd: str,
+    opts: list[str],
+    global_opts: Optional[list[str]] = None,
+    *,
+    include_info: bool = False,
+    include_error: bool = False,
+    user: Optional[str] = None,
+    client: Optional[str] = None,
+) -> list[dict]:
+    records = p4_get_records(
+        cmd, opts, global_opts, include_info=include_info,
+        include_error=include_error, user=user, client=client)
+
+    return [{k.decode(): v.decode() if isinstance(v, bytes) else v
+             for k, v in record.items()}
+            for record in records]
+
+
+def p4_get_client(
+    get_client_name: str,
+    *,
+    user: Optional[str] = None,
+    client: Optional[str] = None,
+) -> Optional[dict]:
+    records = p4_get_decoded_records('client', ['-o', get_client_name],
+                                     user=user, client=client)
+
+    if records:
+        return records[0]
+    else:
+        return None
+
+
+def p4_changes(
+    pathspecs: list[str],
+    *,
+    limit: int = 10,
+    user: Optional[str] = None,
+    client: Optional[str] = None,
+) -> list[dict]:
+    limit = min(limit, 100)
+    opts = ['-s', 'submitted', '-t', '-L', f'-m{limit}', *pathspecs]
+    return p4_get_records('changes', opts, user=user, client=client)
+
+
+def p4_fstat(
+    pathspecs: list[str],
+    *,
+    limit: int = 0,
+    user: Optional[str] = None,
+    client: Optional[str] = None,
+) -> list[dict]:
+    opts: list[str] = []
+
+    if limit > 0:
+        opts.append(f'-m{limit}')
+
+    opts.extend(pathspecs)
+
+    return p4_get_records('fstat', opts, log_unreturned=False,
+                          user=user, client=client)
+
+
+def p4_have(
+    pathspecs: list[str],
+    *,
+    user: Optional[str] = None,
+    client: Optional[str] = None,
+) -> list[dict]:
+    records = p4_get_decoded_records('have', pathspecs, include_error=True,
+                                     user=user, client=client)
+
+    return records
+
+
+@dataclass
+class P4PrintResult:
+    stat: dict
+    contents: Optional[bytearray]
+
+    @property
+    def is_valid(self) -> bool:
+        try:
+            return ((self.stat[b'code'] == b'stat')
+                    and (self.contents is not None))
+        except KeyError:
+            return False
+
+    @property
+    def text(self) -> Optional[str]:
+        if self.contents is not None:
+            return self.contents.decode()
+        else:
+            return None
+
+
+def p4_print(
+    pathspecs: list[str],
+    *,
+    user: Optional[str] = None,
+    client: Optional[str] = None,
+) -> list[P4PrintResult]:
+    '''
+    Returns a `P4PrintResult` for each path in `pathspecs`.
+    If the file does not exist, or is deleted at the specified revision,
+    fileContents will be `None`.
+    '''
+    results: list[P4PrintResult] = []
+    records = p4_get_records('print', pathspecs, include_info=True,
+                             include_error=True, user=user, client=client)
+
+    # For each file, a stat record is returned, then zero or more text records.
+    # Cases where there's no text record include errors or deleted files.
+    for record in records:
+        if (record[b'code'] == b'stat') or (record[b'code'] == b'error'):
+            results.append(P4PrintResult(record, None))
+            continue
+        elif record[b'code'] == b'text':
+            current_result = results[-1]
+            if current_result.contents is None:
+                current_result.contents = bytearray()
+            current_result.contents.extend(record[b'data'])
+            continue
+
+        logging.warning(f'p4_print(): Unhandled record: {record}')
+
+    return results
+
+
+def p4_latest_code_change(
+    paths: list[str],
+    *,
+    range: Optional[str] = None,
+    exts: Optional[list[str]] = None,
+    user: Optional[str] = None,
+    client: Optional[str] = None,
+) -> Optional[int]:
+    '''
+    Code CL determination compatible with UGS/precompiled binaries.
+    See `WorkspaceUpdate.ExecuteAsync` in `WorkspaceUpdate.cs`
+    '''
+
+    ALL_CODE_EXTS = [
+        '.c', '.cc', '.cpp', '.inl', '.m', '.mm', '.rc', '.cs', '.csproj',
+        '.h', '.hpp', '.usf', '.ush', '.uproject', '.uplugin', '.sln']
+
+    if exts is None:
+        exts = ALL_CODE_EXTS
+
+    rangespec = f'@{range}' if range is not None else ''
+
+    code_paths: list[str] = []
+    for path in paths:
+        if path.endswith('...'):
+            code_paths.extend([f'{path}/*{ext}{rangespec}'
+                               for ext in exts])
+
+    code_changes = p4_changes(code_paths, limit=1,
+                              user=user, client=client)
+
+    if len(code_changes) > 0:
+        return max(int(x[b'change']) for x in code_changes)
+    else:
+        return None
diff --git a/Engine/Plugins/VirtualProduction/Switchboard/Source/Switchboard/switchboard/switchboard_logging.py b/Engine/Plugins/VirtualProduction/Switchboard/Source/Switchboard/switchboard/switchboard_logging.py
--- a/Engine/Plugins/VirtualProduction/Switchboard/Source/Switchboard/switchboard/switchboard_logging.py
+++ b/Engine/Plugins/VirtualProduction/Switchboard/Source/Switchboard/switchboard/switchboard_logging.py
@@ -193,7 +193,7 @@ class HTMLogger(object):
 from io import StringIO
 _LOG_STREAM = StringIO()
 logging.basicConfig(stream=_LOG_STREAM)
-_LOGGER = logging.getLogger("switchboard")
+_LOGGER = logging.getLogger()
 DEFAULT_LOG_PATH = os.path.join(tempfile.gettempdir(), "switchboard")
 
 QT_HANDLER = QtHandler()
diff --git a/Engine/Plugins/VirtualProduction/Switchboard/Source/Switchboard/switchboard/switchboard_widgets.py b/Engine/Plugins/VirtualProduction/Switchboard/Source/Switchboard/switchboard/switchboard_widgets.py
--- a/Engine/Plugins/VirtualProduction/Switchboard/Source/Switchboard/switchboard/switchboard_widgets.py
+++ b/Engine/Plugins/VirtualProduction/Switchboard/Source/Switchboard/switchboard/switchboard_widgets.py
@@ -62,7 +62,7 @@ class MultiSelectionComboBox(QtWidgets.QComboBox):
 
     def add_items(self, selected_entries, all_entries):
         for entry in all_entries:
-            self.addItem(entry)
+            self.addItem(str(entry), entry)
             item = self.model().item(self.count()-1, 0)
             item.setFlags(QtCore.Qt.ItemIsUserCheckable | QtCore.Qt.ItemIsEnabled)
             state = QtCore.Qt.Checked if entry in selected_entries else QtCore.Qt.Unchecked
@@ -70,7 +70,7 @@ class MultiSelectionComboBox(QtWidgets.QComboBox):
 
         invalid_entries = [entry for entry in selected_entries if entry not in all_entries]
         for entry in invalid_entries:
-            self.addItem(entry)
+            self.addItem(str(entry), entry)
             item = self.model().item(self.count()-1, 0)
             item.setFlags(QtCore.Qt.ItemIsUserCheckable | QtCore.Qt.ItemIsEnabled)
             item.setCheckState(QtCore.Qt.Checked)
@@ -117,14 +117,16 @@ class MultiSelectionComboBox(QtWidgets.QComboBox):
         # Without this, display can be incorrect when no items are checked [UE-112543].
         self.setCurrentIndex(0)
 
-        selected_entries = []
+        selected_values = []
+        selected_strings = []
         for i in range(self.count()):
             item = self.model().item(i, 0)
             if item.checkState() == QtCore.Qt.Checked:
-                selected_entries.append(self.itemText(i))
+                selected_values.append(self.itemData(i))
+                selected_strings.append(self.itemText(i))
 
-        self.setEditText(self.separator.join(selected_entries))
-        self.signal_selection_changed.emit(selected_entries)
+        self.setEditText(self.separator.join(selected_strings))
+        self.signal_selection_changed.emit(selected_values)
 
 
 # A combo box that has a static icon and opens a list of options when pressed, e.g. view options in settings
diff --git a/Engine/Plugins/VirtualProduction/Switchboard/Source/Switchboard/switchboard/ugs_utils.py b/Engine/Plugins/VirtualProduction/Switchboard/Source/Switchboard/switchboard/ugs_utils.py
--- a/Engine/Plugins/VirtualProduction/Switchboard/Source/Switchboard/switchboard/ugs_utils.py
+++ b/Engine/Plugins/VirtualProduction/Switchboard/Source/Switchboard/switchboard/ugs_utils.py
@@ -1,29 +1,38 @@
 # Copyright Epic Games, Inc. All Rights Reserved.
 
+import csv
 import io
 import logging
-from os import path
-from pathlib import Path
+import os
+import pathlib
 import re
 import shutil
 import subprocess
 import sys
-from typing import Callable, Dict, List, Optional, TextIO, Union
+from typing import Callable, Optional, TextIO, Union
+import uuid
 
-from . import switchboard_utils as sb_utils
+from . import p4_utils, switchboard_utils as sb_utils
+
+try:
+    import P4
+    P4PYTHON_AVAILABLE = True
+except ImportError:
+    P4PYTHON_AVAILABLE = False
 
 
 def find_bin(
-    dir_to_look_under: Optional[Path] = None
-) -> Optional[Path]:
+    dir_to_look_under: Optional[pathlib.Path] = None
+) -> Optional[pathlib.Path]:
     ugs_dir = dir_to_look_under
 
     if not ugs_dir:
         env_ugs_path = shutil.which('ugs')
         if env_ugs_path:
-            ugs_dir = Path(env_ugs_path).parent
+            ugs_dir = pathlib.Path(env_ugs_path).parent
         elif sys.platform.startswith('win'):
-            ugs_dir = path.join('${LOCALAPPDATA}', 'UnrealGameSync', 'Latest')
+            ugs_dir = os.path.join('${LOCALAPPDATA}', 'UnrealGameSync',
+                                   'Latest')
         # elif:
             # for other platforms we don't have a default install location (the
             # dll is usually installed from perforce, which could go anywhere)
@@ -39,24 +48,24 @@ def find_bin(
                       'library (ugs.dll). ' + INSTALL_PATH_SUGGESTION)
         return None
 
-    ugs_bin_path = path.join(path.expandvars(ugs_dir), 'ugs.dll')
-    if not path.exists(ugs_bin_path):
+    ugs_bin_path = os.path.join(os.path.expandvars(ugs_dir), 'ugs.dll')
+    if not os.path.exists(ugs_bin_path):
         logging.error(f"Failed to find '{ugs_bin_path}'. "
                       + INSTALL_PATH_SUGGESTION)
         return None
 
-    return Path(ugs_bin_path)
+    return pathlib.Path(ugs_bin_path)
 
 
 def setup_dependencies(
-    ue_engine_dir: Path
+    ue_engine_dir: pathlib.Path
 ) -> int:
     logging.debug('Verifying UnrealGameSync dependencies.')
 
     if sys.platform.startswith('win'):
         script_name = 'GetDotnetPath.bat'
-        dotnet_setup_script = path.join(ue_engine_dir, 'Build', 'BatchFiles',
-                                        script_name)
+        dotnet_setup_script = os.path.join(ue_engine_dir, 'Build',
+                                           'BatchFiles', script_name)
         dotnet_setup_args = [dotnet_setup_script]
     else:
         if sys.platform.startswith('darwin'):
@@ -64,11 +73,11 @@ def setup_dependencies(
         else:
             platform_dirname = 'Linux'
 
-        platform_scripts_dir = path.join(ue_engine_dir, 'Build', 'BatchFiles',
-                                         platform_dirname)
+        platform_scripts_dir = os.path.join(ue_engine_dir, 'Build',
+                                            'BatchFiles', platform_dirname)
 
         script_name = 'SetupEnvironment.sh'
-        dotnet_setup_script = path.join(platform_scripts_dir, script_name)
+        dotnet_setup_script = os.path.join(platform_scripts_dir, script_name)
         dotnet_setup_args = [dotnet_setup_script, '-dotnet',
                              platform_scripts_dir]
 
@@ -93,8 +102,8 @@ def setup_dependencies(
 
 
 def _find_engine_dir(
-    uproj_path: Path
-) -> Optional[Path]:
+    uproj_path: pathlib.Path
+) -> Optional[pathlib.Path]:
     iter_dir = uproj_path.parent
     while not iter_dir.samefile(iter_dir.anchor):
         candidate = iter_dir / 'Engine'
@@ -107,9 +116,9 @@ def _find_engine_dir(
 
 
 def _get_active_ugs_context(
-    ugs_dll_path: Path,
-    cwd: Optional[Path] = None
-) -> Optional[Dict[str, str]]:
+    ugs_dll_path: pathlib.Path,
+    cwd: Optional[pathlib.Path] = None
+) -> Optional[dict[str, str]]:
     # Assumes UGS dependencies have been already setup (prevents us from
     # redundantly running `setup_dependencies()`)
 
@@ -121,7 +130,7 @@ def _get_active_ugs_context(
                               startupinfo=sb_utils.get_hidden_sp_startupinfo()
                               ) as status_proc:
             for status_line in status_proc.stdout:
-                status_line_str =  status_line.decode().rstrip()
+                status_line_str = status_line.decode().rstrip()
                 logging.debug(f'ugs> {status_line_str}')
                 if 'Project:' in status_line_str:
                     match = re.match(r"Project:\s*//(\S+?)/(\S+)",
@@ -145,9 +154,9 @@ def _get_active_ugs_context(
 
 
 def _set_active_ugs_context(
-    ugs_dll_path: Path,
-    ugs_settings: Dict[str, str],
-    cwd: Optional[Path] = None
+    ugs_dll_path: pathlib.Path,
+    ugs_settings: dict[str, str],
+    cwd: Optional[pathlib.Path] = None
 ) -> int:
     # Assumes UGS dependencies have been already setup (prevents us from
     # redundantly running `setup_dependencies()`)
@@ -172,9 +181,9 @@ def _set_active_ugs_context(
 
 
 def run(
-    ugs_args: Union[str, List[str]],
-    uproj_path: Path,
-    ugs_bin_dir: Optional[Path] = None,
+    ugs_args: Union[str, list[str]],
+    uproj_path: pathlib.Path,
+    ugs_bin_dir: Optional[pathlib.Path] = None,
     user: Optional[str] = None,
     client: Optional[str] = None,
     output_handling_fn: Optional[Callable[[str], int]] = None,
@@ -214,9 +223,9 @@ def run(
     logging.info('Setting UnrealGameSync context.')
     # UGS expects a project path relative to the repo's root
     if ue_engine_dir:
-        sanitized_proj_path = path.relpath(uproj_path, ue_engine_dir.parent)
+        sanitized_proj_path = os.path.relpath(uproj_path, ue_engine_dir.parent)
     else:
-        sanitized_proj_path = uproj_path
+        sanitized_proj_path = str(uproj_path)
     # and only accepts paths with delimited by forward-slashes
     sanitized_proj_path = sanitized_proj_path.replace('\\', '/')
     op_context_params = {'project': sanitized_proj_path}
@@ -272,10 +281,10 @@ def run(
 
 
 def sync(
-    uproj_path: Path,
+    uproj_path: pathlib.Path,
     sync_cl: Optional[int] = None,
     sync_pcbs: Optional[bool] = False,
-    ugs_bin_dir: Optional[Path] = None,
+    ugs_bin_dir: Optional[pathlib.Path] = None,
     user: Optional[str] = None,
     client: Optional[str] = None
 ) -> Optional[int]:
@@ -299,11 +308,11 @@ def sync(
 
 
 def latest_chagelists(
-    uproj_path: Path,
-    ugs_bin_dir: Optional[Path] = None,
+    uproj_path: pathlib.Path,
+    ugs_bin_dir: Optional[pathlib.Path] = None,
     user: Optional[str] = None,
     client: Optional[str] = None
-) -> Optional[List[int]]:
+) -> Optional[list[int]]:
 
     changelist_list = []
 
@@ -330,14 +339,14 @@ def latest_chagelists(
 
 
 def get_depot_config_paths(
-    engine_dir: Union[str, Path],
-    project_dir: Union[str, Path]
-) -> List[str]:
+    engine_dir: Optional[Union[str, pathlib.Path]],
+    project_dir: Optional[Union[str, pathlib.Path]]
+) -> list[str]:
     # See GetDepotConfigPaths in UnrealGameSyncShared/Utility.cs
-    engine_dir = str(engine_dir)
-    project_dir = str(project_dir)
+    engine_dir = str(engine_dir) if engine_dir else None
+    project_dir = str(project_dir) if project_dir else None
 
-    config_paths: List[str] = []
+    config_paths: list[str] = []
 
     def add_platform_dirs(base: str, filename: str):
         config_paths.extend([
@@ -350,11 +359,13 @@ def get_depot_config_paths(
         add_platform_dirs(f'{base}/Platforms/*{rel}', filename)
         add_platform_dirs(f'{base}/Restricted/*{rel}', filename)
 
-    add_platform_extensions(engine_dir, '/Programs/UnrealGameSync',
-                            'UnrealGameSync.ini')
+    if engine_dir:
+        add_platform_extensions(engine_dir, '/Programs/UnrealGameSync',
+                                'UnrealGameSync.ini')
 
-    add_platform_extensions(project_dir, '/Build',
-                            'UnrealGameSync.ini')
+    if project_dir:
+        add_platform_extensions(project_dir, '/Build',
+                                'UnrealGameSync.ini')
 
     return config_paths
 
@@ -367,7 +378,7 @@ class IniParser:
     class Section:
         def __init__(self, name: str):
             self.name = name
-            self.pairs: Dict[str, str] = {}
+            self.pairs: dict[str, str] = {}
 
         def set_value(self, key: str, value: str):
             self.pairs[key] = value
@@ -379,7 +390,7 @@ class IniParser:
                 self.pairs[key] = value
 
     def __init__(self):
-        self.sections: Dict[str, IniParser.Section] = {}
+        self.sections: dict[str, IniParser.Section] = {}
 
     def read_string(self, string: str, source: str):
         sfile = io.StringIO(string)
@@ -394,14 +405,19 @@ class IniParser:
 
         self._read(file, source)
 
-    def try_get(self, section_name: str, key: str) -> Optional[List[str]]:
+    def try_get(self, section_name: str, key: str) -> Optional[list[str]]:
+        value = self.try_get_raw(section_name, key)
+        return value.splitlines() if value is not None else None
+
+    def try_get_raw(self, section_name: str, key: str) -> Optional[str]:
+        ''' Multiple/appended values will be returned separated by '\\n' '''
         section = self.sections.get(section_name)
         if section is None:
             return None
         value = section.pairs.get(key)
         if value is None:
             return None
-        return value.splitlines()
+        return value
 
     def _read(self, file: TextIO, source: str):
         # See Parse() in UnrealGameSyncShared/ConfigFile.cs
@@ -456,3 +472,184 @@ class IniParser:
                     logging.info(f'{key}=')
                     for line in value_lines:
                         logging.info(f'\t{line}')
+
+
+def parse_depot_ugs_configs(
+    engine_dir: Optional[Union[str, pathlib.Path]],
+    project_dir: Optional[Union[str, pathlib.Path]],
+    p4user: Optional[str] = None,
+    p4client: Optional[str] = None,
+) -> IniParser:
+    '''
+    Given local or depot engine and/or project paths, search the Perforce depot
+    for UnrealGameSync.ini configs at known locations, coalesce their contents,
+    and return an `IniParser` populated with the results.
+    '''
+
+    p4opts = {'user': p4user, 'client': p4client}
+
+    logging.debug('Parsing UnrealGameSync configuration from depot')
+
+    # Might append depot wildcards to host paths and end up with mixed syntax
+    # here, but the fstat separates them into normalized depot/host paths.
+    mixed_config_paths = get_depot_config_paths(engine_dir, project_dir)
+
+    config_fstats = p4_utils.p4_fstat(mixed_config_paths, **p4opts)
+
+    config_contents_map: dict[str, Optional[str]] = {}
+
+    # First pass: get contents of locally modified files, or depot paths
+    # of files to `p4 print` in the second pass otherwise
+    for record in config_fstats:
+        action: Optional[bytes] = record.get(b'action')
+        head_action: Optional[bytes] = record.get(b'headAction')
+        if record[b'code'] == b'stat':
+            if (b'clientFile' in record) and (b'action' in record):
+                # If the client file exists and is modified, use local ver
+                client_path = pathlib.Path(record[b'clientFile'].decode())
+                if client_path.exists and action == b'edit':
+                    local_contents = client_path.read_text()
+                    config_contents_map[str(client_path)] = local_contents
+                    continue
+                elif action == b'delete':
+                    continue
+
+            if head_action == b'delete':
+                continue
+
+            if b'depotFile' in record:
+                depot_file = record[b'depotFile'].decode()
+                config_contents_map[depot_file] = None
+                continue
+
+        logging.warning(f'Unhandled fstat record: {record}')
+
+    # Second pass: p4 prints of depot configs @ head
+    depot_print_paths = [k for k, v in config_contents_map.items()
+                         if v is None]  # configs we didn't read locally
+
+    config_prints = p4_utils.p4_print(depot_print_paths, **p4opts)
+
+    for depot_path, record in zip(depot_print_paths, config_prints):
+        if record.is_valid:
+            config_contents_map[depot_path] = record.text
+        else:
+            logging.warning(f'Invalid P4PrintResult for {depot_path}')
+
+    parser = IniParser()
+    for path, contents in config_contents_map.items():
+        if contents is None:
+            logging.warning(f'No contents for {path}')
+            continue
+
+        parser.read_string(contents, path)
+
+    return parser
+
+
+class SyncFilters:
+    '''
+    Determines whether given paths are excluded, based on a list of rules
+    containing Perforce-style wildcards.
+
+    These rules can be extracted from UnrealGameSync.ini files (see
+    `read_categories_from_ini_parser()`), or added via calls directly to
+    `SyncFilters.map.insert()`.
+
+    To test a path against the filters, use `SyncFilters.includes_path()`.
+
+    Examples:
+        - `-/ExcludedRoot/...`
+        - `-*.excludedfileextension`
+        - `-.../SubdirectoryExcludedAnywhere/...`
+    '''
+
+    @classmethod
+    def supported(cls):
+        return P4PYTHON_AVAILABLE
+
+    class CsvDialect(csv.excel):
+        skipinitialspace = True
+
+    class Category:
+        def __init__(self, id: uuid.UUID, name: str, paths: list[str]):
+            self.id = id
+            self.name = name
+            self.paths = paths
+
+            self.hidden = False
+            self.requires: Optional[uuid.UUID] = None
+
+    def __init__(self):
+        self.categories: dict[uuid.UUID, SyncFilters.Category] = {}
+        self.excluded_categories: set[uuid.UUID] = set()
+        self.map = P4.Map()
+
+        self._reset_map()
+
+    def _reset_map(self):
+        self.map.clear()
+        self.map.insert('...')
+
+    def read_categories_from_ini_parser(self, config: IniParser):
+        categories_raw = config.try_get_raw('Options', 'SyncCategory')
+        if categories_raw:
+            sfile = io.StringIO(categories_raw)
+            self._read(sfile)
+
+    def exclude_category(self, category_id: uuid.UUID):
+        if category_id in self.excluded_categories:
+            return  # already excluded
+
+        category = self.categories.get(category_id)
+        if category is None:
+            logging.warning(f'exclude_category: no such category {category_id}')
+            return
+
+        self._map_exclusion(category)
+
+    def includes_path(self, path: str) -> bool:
+        return self.map.includes(path)
+
+    def _read(self, file: TextIO):
+        category_lines = [line.strip('()\r\n') for line in file]
+        category_reader = csv.reader(category_lines, SyncFilters.CsvDialect)
+        for category_fields in category_reader:
+            category_dict: dict[str, str] = {}
+            for field in category_fields:
+                key, _, value = field.partition('=')
+                value = value.strip('"')
+                category_dict[key] = value
+
+            id = uuid.UUID(hex=category_dict['UniqueId'])
+            name = category_dict['Name']
+            paths = category_dict['Paths'].split(';')
+            paths = [path.lstrip() for path in paths]
+
+            hidden = category_dict.get('Hidden', '').lower() == 'true'
+            requires = category_dict.get('Requires')
+            requires = uuid.UUID(requires) if requires else None
+
+            if id not in self.categories:
+                category = self.create_category(id, name, paths)
+                category.hidden = hidden
+                category.requires = requires
+            else:
+                # This seems weird, but... some configs need it
+                self.categories[id].paths.extend(paths)
+
+    def create_category(self, id: uuid.UUID, name: str, paths: list[str]):
+        category = SyncFilters.Category(id, name, paths)
+        self.categories[id] = category
+        return category
+
+    DIR_REGEX = re.compile(r'(\.\.\.|/|\\)')
+
+    def _map_exclusion(self, category: Category):
+        self.excluded_categories.add(category.id)
+        for path in category.paths:
+            self.map.insert(f'-{path}')
+
+            if not self.DIR_REGEX.search(path):
+                # unqualified file pattern, so also match nested instances
+                self.map.insert(f'-.../{path}')
diff --git a/Engine/Plugins/VirtualProduction/Switchboard/Source/Switchboard/tests/test_ugs_utils.py b/Engine/Plugins/VirtualProduction/Switchboard/Source/Switchboard/tests/test_ugs_utils.py
new file mode 100644
--- /dev/null
+++ b/Engine/Plugins/VirtualProduction/Switchboard/Source/Switchboard/tests/test_ugs_utils.py
@@ -0,0 +1,123 @@
+# Copyright Epic Games, Inc. All Rights Reserved.
+
+import unittest
+import uuid
+
+from switchboard import ugs_utils
+
+
+# To run these tests:
+# > cd ...\Engine\Plugins\VirtualProduction\Switchboard\Source\Switchboard
+# > ..\..\..\..\..\Extras\ThirdPartyNotUE\SwitchboardThirdParty\Python\Scripts\activate
+# (switchboard_venv) > python -m unittest discover -v -s tests\
+
+
+class IniParserTestCase(unittest.TestCase):
+    def setUp(self):
+        self.parser = ugs_utils.IniParser()
+
+    def add_ini_sources(self, sources: dict[str, str]):
+        # Takes a dict of `filename -> file contents` and loads them as INIs.
+        for source_name, source_contents in sources.items():
+            self.parser.read_string(source_contents, source_name)
+
+    def test_basic(self):
+        self.add_ini_sources({
+            'Basic.ini':
+                '[Section]\n'
+                'Key=Value\n',
+        })
+
+        self.assertEqual(
+            self.parser.try_get('Section', 'Key'),
+            ['Value']
+        )
+
+        self.assertEqual(
+            self.parser.try_get('Section', 'NonexistentKey'),
+            None
+        )
+
+        self.assertEqual(
+            self.parser.try_get('NonexistentSection', 'Key'),
+            None
+        )
+
+    def test_append(self):
+        self.add_ini_sources({
+            'AppendOne.ini':
+                '[Section]\n'
+                '+SingleSourceKey=ValueA\n'
+                '+SingleSourceKey=ValueB\n'
+                'TwoSourceKey=ValueOne',
+            'AppendTwo.ini':
+                '[Section]\n'
+                '+TwoSourceKey=ValueTwo',
+        })
+
+        self.assertEqual(
+            self.parser.try_get('Section', 'SingleSourceKey'),
+            ['ValueA', 'ValueB'],
+            'Appending a value within the same file'
+        )
+
+        self.assertEqual(
+            self.parser.try_get('Section', 'TwoSourceKey'),
+            ['ValueOne', 'ValueTwo'],
+            'Appending a value from a subsequent file'
+        )
+
+    def test_overwrite(self):
+        self.add_ini_sources({
+            'Original.ini':
+                '[Section]\n'
+                'Key=OriginalValue\n',
+            'Overwrite.ini':
+                '[Section]\n'
+                'Key=OverwrittenValue\n',
+        })
+
+        self.assertEqual(
+            self.parser.try_get('Section', 'Key'),
+            ['OverwrittenValue'],
+            'Overwrite.ini value should replace Original.ini value'
+        )
+
+
+class SyncFiltersTestCase(unittest.TestCase):
+    @classmethod
+    def setUpClass(cls):
+        if not ugs_utils.SyncFilters.supported():
+            raise RuntimeError('SyncFilters not available')
+
+    def setUp(self):
+        self.sync_filters = ugs_utils.SyncFilters()
+
+    def add_new_category(self, name: str, paths: list[str]) -> uuid.UUID:
+        id = uuid.uuid4()
+        self.sync_filters.create_category(id, name, paths)
+        return id
+
+    def test_exclusions(self):
+        DEPOT = [
+            '/GenerateProjectFiles.bat',
+            '/Engine/Content/Maps/Entry.umap',
+            '/Engine/Source/UnrealEditor.Target.cs',
+        ]
+
+        for path in DEPOT:
+            self.assertTrue(self.sync_filters.includes_path(path))
+
+        content_id = self.add_new_category('Content', ['*.uasset', '*.umap'])
+        self.sync_filters.exclude_category(content_id)
+
+        self.assertEqual(list(map(self.sync_filters.includes_path, DEPOT)),
+                         [True, False, True],
+                         'Should exclude only the umap file')
+
+        source_id = self.add_new_category('Source', ['/Engine/Source/...'])
+        self.sync_filters.exclude_category(source_id)
+
+        self.assertEqual(list(map(self.sync_filters.includes_path, DEPOT)),
+                         [True, False, False],
+                         'Should now exclude the umap and source files')
