Switchboard: Recording and OSC improvements.

- Sets local Take Recorder "Record On Client" equal to the value of Switchboard "Is Recording Device" via command line.

- Removes global "OSC Client Port" setting, which previously determined the remote destination port used for all outgoing OSC connections.
- Introduces as a replacement DeviceUnreal + DeviceLiveLinkFace "OSC Port" class settings, with appropriate defaults for each integration, and per-device overrides.
- Clarifies terminology and adds tooltips to all OSC-related settings.

Other fixes:
- Removed unsafe mutable default on ApplicationLauncher.launch (Tools menu) which would erroneously accumulate args from all prior calls.
- Fixes a bug where launching SwitchboardListener from the tools menu would result in a broken console window + inherited stdout.

diff --git a/Engine/Plugins/VirtualProduction/Switchboard/Source/Switchboard/switchboard/config.py b/Engine/Plugins/VirtualProduction/Switchboard/Source/Switchboard/switchboard/config.py
--- a/Engine/Plugins/VirtualProduction/Switchboard/Source/Switchboard/switchboard/config.py
+++ b/Engine/Plugins/VirtualProduction/Switchboard/Source/Switchboard/switchboard/config.py
@@ -2159,20 +2159,17 @@ class Config(object):
             self.basic_project_settings["content_plugin_filters"])
 
         self.osc_settings = {
-            "osc_server_port": IntSetting(
-                "osc_server_port",
-                "OSC Server Port",
-                data.get('osc_server_port', 6000)
-            ),
-            "osc_client_port": IntSetting(
-                "osc_client_port",
-                "OSC Client Port",
-                data.get('osc_client_port', 8000)
+            'osc_server_port': IntSetting(
+                attr_name='osc_server_port',
+                nice_name='Internal OSC Server Port',
+                value=data.get('osc_server_port', 6000),
+                tool_tip=(
+                    'The port on which Switchboard listens for incoming OSC '
+                    'connections from other devices/applications.')
             )
         }
 
         self.OSC_SERVER_PORT = self.osc_settings["osc_server_port"]
-        self.OSC_CLIENT_PORT = self.osc_settings["osc_client_port"]
 
         self.source_control_settings = {
             "p4_enabled": BoolSetting(
@@ -2447,7 +2444,6 @@ class Config(object):
 
         # OSC settings
         data["osc_server_port"] = self.OSC_SERVER_PORT.get_value()
-        data["osc_client_port"] = self.OSC_CLIENT_PORT.get_value()
 
         # Source Control Settings
         data["p4_enabled"] = self.P4_ENABLED.get_value()
diff --git a/Engine/Plugins/VirtualProduction/Switchboard/Source/Switchboard/switchboard/devices/device_base.py b/Engine/Plugins/VirtualProduction/Switchboard/Source/Switchboard/switchboard/devices/device_base.py
--- a/Engine/Plugins/VirtualProduction/Switchboard/Source/Switchboard/switchboard/devices/device_base.py
+++ b/Engine/Plugins/VirtualProduction/Switchboard/Source/Switchboard/switchboard/devices/device_base.py
@@ -84,6 +84,8 @@ class Device(QtCore.QObject):
         self.device_qt_handler = DeviceQtHandler()
 
         self.setting_address = AddressSetting('address', 'Address', address)
+        self.setting_address.signal_setting_changed.connect(
+            lambda: self._check_recreate_osc_client())
 
         # override any setting that was passed via kwargs
         for setting in self.setting_overrides():
@@ -105,7 +107,7 @@ class Device(QtCore.QObject):
         self.device_hash = random.getrandbits(64)
 
         # Lazily create the OSC client as needed
-        self.osc_client = None
+        self.osc_client: Optional[pythonosc.udp_client.SimpleUDPClient] = None
 
         self.device_recording = None
 
@@ -213,9 +215,6 @@ class Device(QtCore.QObject):
     @address.setter
     def address(self, value: str):
         self.setting_address.update_value(value)
-        # todo-dara: probably better to have the osc client connect to a
-        # change of the address.
-        self.setup_osc_client(CONFIG.OSC_CLIENT_PORT.get_value())
 
     @property
     def status(self):
@@ -252,7 +251,7 @@ class Device(QtCore.QObject):
         self._engine_changelist = value
         self.device_qt_handler.signal_device_engine_changelist_changed.emit(
             self)
-        
+
     @property
     def built_engine_changelist(self):
         return self._built_engine_changelist
@@ -264,6 +263,10 @@ class Device(QtCore.QObject):
             self
         )
 
+    @property
+    def device_osc_port(self) -> int:
+        raise NotImplementedError()
+
     def set_slate(self, value):
         pass
 
@@ -308,15 +311,26 @@ class Device(QtCore.QObject):
 
     def setup_osc_client(self, osc_port):
         try:
+            LOGGER.osc(
+                f'{self.name}: (Re)creating OSC client '
+                f'{self.address}:{osc_port}')
+
             self.osc_client = pythonosc.udp_client.SimpleUDPClient(
                 self.address, osc_port)
         except socket.gaierror as exc:
             LOGGER.error(f'{self.name}: Invalid OSC server address',
                          exc_info=exc)
 
+    def _check_recreate_osc_client(self):
+        if self.osc_client is not None:
+            address_changed = self.osc_client._address != self.address
+            port_changed = self.osc_client._port != self.device_osc_port
+            if address_changed or port_changed:
+                self.setup_osc_client(self.device_osc_port)
+
     def send_osc_message(self, command, value, log=True):
         if not self.osc_client:
-            self.setup_osc_client(CONFIG.OSC_CLIENT_PORT.get_value())
+            self.setup_osc_client(self.device_osc_port)
 
         if log:
             LOGGER.osc(
diff --git a/Engine/Plugins/VirtualProduction/Switchboard/Source/Switchboard/switchboard/devices/livelinkface/plugin_livelinkface.py b/Engine/Plugins/VirtualProduction/Switchboard/Source/Switchboard/switchboard/devices/livelinkface/plugin_livelinkface.py
--- a/Engine/Plugins/VirtualProduction/Switchboard/Source/Switchboard/switchboard/devices/livelinkface/plugin_livelinkface.py
+++ b/Engine/Plugins/VirtualProduction/Switchboard/Source/Switchboard/switchboard/devices/livelinkface/plugin_livelinkface.py
@@ -1,23 +1,32 @@
 # Copyright Epic Games, Inc. All Rights Reserved.
+
+import datetime
+
+from PySide2 import QtCore, QtGui
+
 from switchboard import config_osc as osc
-from switchboard.config import CONFIG, SETTINGS
+from switchboard.config import CONFIG, SETTINGS, IntSetting
 from switchboard.devices.device_base import Device, DeviceStatus
 from switchboard.devices.device_widget_base import DeviceWidget
 from switchboard.switchboard_logging import LOGGER
 
-from PySide2 import QtCore, QtGui, QtWidgets
-from PySide2.QtCore import QDirIterator
-
-import datetime
 
 class DeviceLiveLinkFace(Device):
-    QUERY_TIME = 5
-
-    def __init__(self, name, address, port=CONFIG.OSC_CLIENT_PORT.get_value(), **kwargs):
+    csettings = {
+        'osc_port': IntSetting(
+            attr_name='osc_port',
+            nice_name='OSC Port',
+            value=8000,
+            tool_tip=(
+                'Must match the port on which the Live Link Face app is '
+                'configured to listen for OSC connections.'),
+        )
+    }
+
+    def __init__(self, name, address, **kwargs):
         super().__init__(name, address, **kwargs)
 
         self.auto_connect = False
-        self.port = port
         self._battery = 1.0
         self.look_for_device = True
 
@@ -27,6 +36,21 @@ class DeviceLiveLinkFace(Device):
         self.query_timer.start(sleep_time_in_ms)
         self._query_time = datetime.datetime.now()
 
+        osc_port_setting = DeviceLiveLinkFace.csettings['osc_port']
+        osc_port_setting.signal_setting_changed.connect(
+            lambda: self._check_recreate_osc_client())
+        osc_port_setting.signal_setting_overridden.connect(
+            lambda: self._check_recreate_osc_client())
+
+    def setting_overrides(self):
+        return super().setting_overrides() + [
+            DeviceLiveLinkFace.csettings['osc_port'],
+        ]
+
+    @property
+    def device_osc_port(self) -> int:
+        return DeviceLiveLinkFace.csettings['osc_port'].get_value(self.name)
+
     @property
     def battery(self):
         return self._battery
@@ -122,11 +146,6 @@ class DeviceLiveLinkFace(Device):
 
         self.send_osc_message(osc.RECORD_START, [slate, take])
 
-    def get_config(self):
-        device_config = super().get_config()
-        device_config['port'] = 8000
-        return device_config
-
 
 class DeviceWidgetLiveLinkFace(DeviceWidget):
     def __init__(self, name, device_hash, address, icons, parent=None):
@@ -146,15 +165,14 @@ class DeviceWidgetLiveLinkFace(DeviceWidget):
         self.look_for_button.clicked.connect(self.look_for_button_clicked)
 
     def set_battery(self, value):
-
         rounded_value = int(round(value * 100 / 10.0)) * 10
-        
+
         qrc_path = f":/icons/images/icon_livelinkface_{rounded_value}.png"
         pixmap = QtGui.QPixmap(qrc_path)
 
         if pixmap.isNull():
             pixmap = self.icon_for_state("enabled").pixmap(QtCore.QSize(40, 40))
-            
+
         self.device_icon.setPixmap(pixmap)
 
     def look_for_button_clicked(self):
diff --git a/Engine/Plugins/VirtualProduction/Switchboard/Source/Switchboard/switchboard/devices/unreal/plugin_unreal.py b/Engine/Plugins/VirtualProduction/Switchboard/Source/Switchboard/switchboard/devices/unreal/plugin_unreal.py
--- a/Engine/Plugins/VirtualProduction/Switchboard/Source/Switchboard/switchboard/devices/unreal/plugin_unreal.py
+++ b/Engine/Plugins/VirtualProduction/Switchboard/Source/Switchboard/switchboard/devices/unreal/plugin_unreal.py
@@ -515,6 +515,14 @@ class DeviceUnreal(Device):
             value=2980,
             tool_tip="Port of SwitchboardListener"
         ),
+        'osc_port': IntSetting(
+            attr_name='osc_port',
+            nice_name='OSC Port',
+            value=5500,
+            tool_tip=(
+                'Must match the port on which the Unreal Editor OSC server '
+                'is configured to listen for OSC connections.'),
+        ),
         'roles_filename': StringSetting(
             attr_name="roles_filename",
             nice_name="Roles Filename",
@@ -936,6 +944,12 @@ class DeviceUnreal(Device):
         CONFIG.ENGINE_SYNC_METHOD.signal_setting_changed.connect(
             self.on_engine_sync_method_changed)
 
+        osc_port_setting = DeviceUnreal.csettings['osc_port']
+        osc_port_setting.signal_setting_changed.connect(
+            lambda: self._check_recreate_osc_client())
+        osc_port_setting.signal_setting_overridden.connect(
+            lambda: self._check_recreate_osc_client())
+
         self.auto_connect = False
 
         self.runtime_str = ""
@@ -1123,6 +1137,7 @@ class DeviceUnreal(Device):
     def setting_overrides(self):
         overrides = super().setting_overrides() + [
             Device.csettings['is_recording_device'],
+            DeviceUnreal.csettings['osc_port'],
             DeviceUnreal.csettings['command_line_arguments'],
             DeviceUnreal.csettings['exec_cmds'],
             DeviceUnreal.csettings['dp_cvars'],
@@ -1275,6 +1290,10 @@ class DeviceUnreal(Device):
     def on_setting_exclude_from_build_changed(self, exclude_from_build):
         self.widget.update_exclude_from_build(exclude_from_build, not self.is_disconnected)
 
+    @property
+    def device_osc_port(self) -> int:
+        return DeviceUnreal.csettings['osc_port'].get_value(self.name)
+
     def set_slate(self, value):
         if not self.is_recording_device:
             return
@@ -1302,7 +1321,7 @@ class DeviceUnreal(Device):
             os.path.dirname(uproject_path), "Config", "Tags", roles_filename)
         _, msg = message_protocol.create_copy_file_from_listener_message(roles_file_path)
         self.unreal_client.send_message(msg)
-        
+
     def _request_unreal_editor_version_file(self):
         '''
         Requests the Engine/Binaries/[platform]/UnrealEditor.version file.
@@ -1890,9 +1909,7 @@ class DeviceUnreal(Device):
     def generate_unreal_command_line_args(self, map_name):
         command_line_args = f'{self.extra_cmdline_args_setting}'
 
-        command_line_args += f' Log={self.log_filename}'
-
-        command_line_args += " "
+        command_line_args += f' Log={self.log_filename} '
 
         if CONFIG.MUSERVER_AUTO_JOIN.get_value() and self.autojoin_mu_server.get_value():
             command_line_args += (
@@ -1919,7 +1936,6 @@ class DeviceUnreal(Device):
             exec_cmds.append(self.exec_command_for_livelink_preset(livelink_preset_gamepath))
 
         # Exec Commands
-
         exec_cmds = [cmd for cmd in exec_cmds if len(cmd.strip())]
 
         if len(exec_cmds):
@@ -2002,6 +2018,14 @@ class DeviceUnreal(Device):
                 ' -UDPMESSAGING_TRANSPORT_STATIC='
                 f'"{",".join(static_endpoints)}"')
 
+        record_on_client = 'True' if self.is_recording_device else 'False'
+        ini_engine = (
+            ' -ini:Engine:'
+            '[/Script/ConcertTakeRecorder.ConcertSessionRecordSettings]:'
+            f'LocalSettings=(bRecordOnClient={record_on_client})'
+        )
+        command_line_args += ini_engine
+
         return (
             f'"{CONFIG.UPROJECT_PATH.get_value(self.name)}" {map_name} '
             f'{command_line_args}')
diff --git a/Engine/Plugins/VirtualProduction/Switchboard/Source/Switchboard/switchboard/tools/application_launcher.py b/Engine/Plugins/VirtualProduction/Switchboard/Source/Switchboard/switchboard/tools/application_launcher.py
--- a/Engine/Plugins/VirtualProduction/Switchboard/Source/Switchboard/switchboard/tools/application_launcher.py
+++ b/Engine/Plugins/VirtualProduction/Switchboard/Source/Switchboard/switchboard/tools/application_launcher.py
@@ -1,47 +1,58 @@
 # Copyright Epic Games, Inc. All Rights Reserved.
 
-import subprocess, os, threading
-
-from typing import List, Optional
+import os
+import subprocess
+import threading
+from typing import Optional
 
 from switchboard import switchboard_utils as sb_utils
-
 from switchboard.switchboard_logging import LOGGER
 
+
 class ApplicationLauncher(object):
     ''' Manages launching external applications '''
 
-    def __init__(self, name:str):
+    def __init__(self, name: str):
         self.lock = threading.Lock()
         self.process: Optional[subprocess.Popen] = None
         self.name = name
 
     def exe_path(self):
         ''' Returns the expected executable path. Extension not included '''
-
         return ""
 
     def exe_name(self):
         return os.path.split(self.exe_path())[1]
 
-    def launch(self, args:Optional[List[str]] = [], allow_duplicate:bool = False):
+    def launch(
+        self,
+        args: Optional[list[str]] = None,
+        allow_duplicate: bool = False
+    ):
         ''' Launches this application with the given arguments '''
+        if args is None:
+            args = []
 
         args.insert(0, self.exe_path())
         cmdline = ' '.join(args)
 
         with self.lock:
-
             if (not allow_duplicate) and self.is_running():
                 return False
 
             if not os.path.exists(self.exe_path()):
-                LOGGER.error(f"Could not find '{self.name}' at '{self.exe_path()}'. Has it been built?")
+                LOGGER.error(
+                    f"Could not find '{self.name}' at '{self.exe_path()}'. "
+                    "Has it been built?")
                 return False
 
             LOGGER.debug(f"Launching '{cmdline}' ...")
 
-            subprocess.Popen(cmdline)
+            creationflags = 0
+            if hasattr(subprocess, 'CREATE_NEW_CONSOLE'):
+                creationflags = subprocess.CREATE_NEW_CONSOLE
+
+            subprocess.Popen(cmdline, creationflags=creationflags)
 
         return True
 
