Switchboard: Added LogDisplayClusterStageMonitoringDWM to default log options.

#skipci
#rnx

diff --git a/Engine/Plugins/VirtualProduction/Switchboard/Source/Switchboard/switchboard/devices/ndisplay/plugin_ndisplay.py b/Engine/Plugins/VirtualProduction/Switchboard/Source/Switchboard/switchboard/devices/ndisplay/plugin_ndisplay.py
--- a/Engine/Plugins/VirtualProduction/Switchboard/Source/Switchboard/switchboard/devices/ndisplay/plugin_ndisplay.py
+++ b/Engine/Plugins/VirtualProduction/Switchboard/Source/Switchboard/switchboard/devices/ndisplay/plugin_ndisplay.py
@@ -430,6 +430,7 @@ class DevicenDisplay(DeviceUnreal):
                 'LogDisplayClusterRenderSync',
                 'LogDisplayClusterViewport',
                 'LogDisplayClusterMedia',
+                'LogDisplayClusterStageMonitoringDWM',
                 'LogLiveLink',
                 'LogRemoteControl',
             ],
