Switchboard: Multi-user server launching / status monitoring fixes.

- Don't manipulate Qt UI from background thread, which was causing intermittent crashes on multi-user launch/shutdown.
- Use `threading.Thread` instead of `QThreadPool` so that the thread is debuggable.
- Build MU args to `subprocess.Popen` as an array instead of a string, so that we can stop using `shell=True` (which was only necessary on Linux, but prevented `process.kill()` from working correctly on Windows).
- Fix `PollProcess.args` code that extracts command lines from sibling processes on Windows.
- Added a helper class to extract command line arguments, with accompanying tests. Running server now correctly validates e.g. server names containing hyphens, and hostnames in addition to IP addresses.

diff --git a/Engine/Plugins/VirtualProduction/Switchboard/Source/Switchboard/switchboard/switchboard_application.py b/Engine/Plugins/VirtualProduction/Switchboard/Source/Switchboard/switchboard/switchboard_application.py
--- a/Engine/Plugins/VirtualProduction/Switchboard/Source/Switchboard/switchboard/switchboard_application.py
+++ b/Engine/Plugins/VirtualProduction/Switchboard/Source/Switchboard/switchboard/switchboard_application.py
@@ -3,12 +3,13 @@
 import os
 import pathlib
 import re
+import shlex
 import subprocess
 import sys
 import tempfile
 import threading
 
-from typing import Dict, List, NamedTuple, Optional
+from typing import NamedTuple, Optional, Union
 
 import pythonosc.dispatcher
 import pythonosc.osc_server
@@ -17,6 +18,7 @@ from .config import CONFIG, SETTINGS
 from .switchboard_logging import LOGGER
 from . import switchboard_utils as sb_utils
 
+
 class OscServer:
     def __init__(self):
         self.server_thread = None
@@ -77,13 +79,61 @@ class OscServer:
         return self.server_thread.is_alive()
 
     def dispatcher_map(self, command, method):
-        # LOGGER.osc(f'OSC Server: dispatcher map {command} to {method}')
         self.dispatcher.map(command, method, needs_reply_address=True)
 
     def _default_callback(self, client_address, command, *args):
         LOGGER.warning(f'Received unhandled OSC message: {command} {args}.')
 
 
+class CommandLineExtractor:
+    '''
+    Given command line arguments as either a single string or an array,
+    tokenizes and normalizes them (quotes removed), and provides a means to
+    extract specific values.
+    '''
+
+    # Any single backslash not followed by a "
+    WIN_SLASH_RE = re.compile(r"(?<!\\)\\(?![\"\\])")
+
+    def __init__(self, args: Union[str, list[str]]):
+        if sys.platform.startswith('win'):
+            # Escape path separators (which shlex will subsequently unescape)
+            if isinstance(args, str):
+                args = self.WIN_SLASH_RE.sub(R'\\\\', args)
+            else:
+                for arg in args:
+                    arg = self.WIN_SLASH_RE.sub(R'\\\\', arg)
+
+        self.tokens = self._extract(args)
+
+    def get_value_for_switch(self, switch: str) -> Optional[str]:
+        switch_re = re.compile(f'-{switch}=(.*)', re.IGNORECASE)
+        for token in self.tokens:
+            if match := switch_re.search(token):
+                return match.group(1)
+
+        return None
+
+    def _extract(self, args: Union[str, list[str]]):
+        tokens: list[str] = []
+
+        if isinstance(args, str):
+            lex = self._shlex(args)
+            tokens = list(lex)
+        else:
+            for arg in args:
+                tokens.extend(list(self._shlex(arg)))
+
+        return tokens
+
+    def _shlex(self, instream) -> shlex.shlex:
+        lex = shlex.shlex(instream=instream, posix=True)
+        lex.whitespace_split = True
+        if sys.platform.startswith('win'):
+            lex.quotes = '"'  # Single quotes don't affect tokenization
+        return lex
+
+
 class MultiUserApplication:
     def __init__(self):
         self.lock = threading.Lock()
@@ -102,11 +152,11 @@ class MultiUserApplication:
     def get_mu_server_multicast_arg(self):
         multicast = CONFIG.MUSERVER_MULTICAST_ENDPOINT.get_value().strip()
         if multicast:
-            return f'-UDPMESSAGING_TRANSPORT_MULTICAST={multicast}'
+            return f'-UDPMESSAGING_TRANSPORT_MULTICAST="{multicast}"'
         return ''
 
     def get_mu_server_endpoint_arg(self):
-        endpoint = self.endpoint_address()
+        endpoint = self.configured_endpoint()
         if endpoint:
             return f'-UDPMESSAGING_TRANSPORT_UNICAST="{endpoint}"'
         return ''
@@ -123,7 +173,7 @@ class MultiUserApplication:
         # and we'd like to pick up on name changes.
         return sb_utils.PollProcess(self.exe_name())
 
-    def launch(self, args: Optional[List[str]] = None):
+    def launch(self, args: Optional[list[str]] = None):
         if args is None:
             args = []
 
@@ -136,83 +186,64 @@ class MultiUserApplication:
                              f'{self.exe_path()}. Has it been built?')
                 return
 
-            cmdline = ''
+            creationflags = 0
+
             if sys.platform.startswith('win'):
-                if CONFIG.MUSERVER_SLATE_MODE.get_value():
-                    cmdline = f'"{self.exe_path()}"'
-                else:
-                    cmdline = f'start "Multi User Server" "{self.exe_path()}"'
+                args.insert(0, f'"{self.exe_path()}"')
+                if not CONFIG.MUSERVER_SLATE_MODE.get_value():
+                    creationflags = subprocess.CREATE_NEW_CONSOLE
             else:
-                cmdline = f'"{self.exe_path()}"'
+                args.insert(0, f'{self.exe_path()}')
 
-            cmdline += f' -CONCERTSERVER="{CONFIG.MUSERVER_SERVER_NAME.get_value()}"'
-            cmdline += f' {CONFIG.MUSERVER_COMMAND_LINE_ARGUMENTS.get_value()}'
-            cmdline += f' {self.get_mu_server_endpoint_arg()}'
-            cmdline += f' {self.get_mu_server_multicast_arg()}'
+            args.extend([
+                f'-CONCERTSERVER="{CONFIG.MUSERVER_SERVER_NAME.get_value()}"',
+                f'{CONFIG.MUSERVER_COMMAND_LINE_ARGUMENTS.get_value()}',
+                f'{self.get_mu_server_endpoint_arg()}',
+                f'{self.get_mu_server_multicast_arg()}'])
 
             if self.concert_ignore_cl:
-                cmdline += " -ConcertIgnore"
+                args.append('-ConcertIgnore')
 
-            if CONFIG.MUSERVER_WORKING_DIR.get_value():
-                cmdline += f' -ConcertWorkingDir="{CONFIG.MUSERVER_WORKING_DIR.get_value()}"'
+            working_dir = CONFIG.MUSERVER_WORKING_DIR.get_value()
+            if working_dir:
+                args.append(f'-ConcertWorkingDir="{working_dir}"')
 
-            if CONFIG.MUSERVER_ARCHIVE_DIR.get_value():
-                cmdline += f' -ConcertSavedDir="{CONFIG.MUSERVER_ARCHIVE_DIR.get_value()}"'
+            archive_dir = CONFIG.MUSERVER_ARCHIVE_DIR.get_value()
+            if archive_dir:
+                args.append(f'-ConcertSavedDir="{archive_dir}"')
 
             if CONFIG.MUSERVER_CLEAN_HISTORY.get_value():
-                cmdline += " -ConcertClean"
+                args.append('-ConcertClean')
 
-            if len(args) > 0:
-                cmdline += f' {" ".join(args)}'
+            args = list(filter(None, args))
 
-            LOGGER.debug(cmdline)
-            self._process = subprocess.Popen(
-                cmdline, shell=True,
-                startupinfo=sb_utils.get_hidden_sp_startupinfo())
-
-            return True
+            if sys.platform.startswith('win'):
+                # Given a list on Windows, Popen calls subprocess.list2cmdline
+                # to combine to a string. However, that also has the effect
+                # of turning `-foo="bar baz"`` into `"-foo=\"bar baz\""`,
+                # and Unreal's FParse won't accept that (OR `"-foo=bar baz"`);
+                # its tokenization is too different from MSVCRT.
+                args = ' '.join(args)
 
-    def terminate(self, bypolling=False):
-        if not bypolling and self._process:
-            self._process.terminate()
-        else:
-            self.poll_process().kill()
+            self._process = subprocess.Popen(args, creationflags=creationflags)
 
-    FIND_IP_RE = re.compile(
-        r'-UDPMESSAGING_TRANSPORT_UNICAST="(\d+.\d+.\d+.\d+:\d+)"',
-        re.IGNORECASE)
-    FIND_NAME_RE = re.compile(
-        r'-CONCERTSERVER="([A-Za-z0-9_]+)"', re.IGNORECASE)
+            return True
 
     def extract_process_info(self):
         pid = self.process.pid
-        command_line = str(self.process.args)
+        args = self.process.args
 
-        if command_line == '' or pid == self._running_pid:
+        if (not args) or (pid == self._running_pid):
             return
 
         self.clear_process_info()
         self._running_pid = pid
-        try:
-            ip_match = self.FIND_IP_RE.search(command_line)
-            name_match = self.FIND_NAME_RE.search(command_line)
-            if ip_match:
-                self._running_endpoint = ip_match.group(1)
-            if name_match:
-                self._running_name = name_match.group(1)
-        except:
-            pass
-
-    def validate_process(self):
-        if self._running_pid is None:
-            return False
 
-        endpoint = self.endpoint_address()
-        server_name = CONFIG.MUSERVER_SERVER_NAME.get_value()
-        if endpoint == self._running_endpoint and server_name == self._running_name:
-            return True
-
-        return False
+        extractor = CommandLineExtractor(args)
+        self._running_endpoint = extractor.get_value_for_switch(
+            'UDPMESSAGING_TRANSPORT_UNICAST') or ''
+        self._running_name = extractor.get_value_for_switch(
+            'CONCERTSERVER') or ''
 
     def running_server_name(self):
         return self._running_name
@@ -220,7 +251,7 @@ class MultiUserApplication:
     def running_endpoint(self):
         return self._running_endpoint
 
-    def endpoint_address(self):
+    def configured_endpoint(self):
         endpoint_setting = CONFIG.MUSERVER_ENDPOINT.get_value().strip()
         endpoint = ""
         if endpoint_setting:
@@ -228,7 +259,7 @@ class MultiUserApplication:
             endpoint = sb_utils.expand_endpoint(endpoint_setting, addr_setting)
         return endpoint
 
-    def server_name(self):
+    def configured_server_name(self):
         return CONFIG.MUSERVER_SERVER_NAME.get_value()
 
     def clear_process_info(self):
@@ -292,7 +323,7 @@ hosts allow = {allowed_addrs}
         self.log_file = tempfile.NamedTemporaryFile(
             prefix='sb_rsync', suffix='.log', mode='rt')
         self.process: Optional[subprocess.Popen] = None
-        self.allowed_clients: Dict[object, RsyncServer.Client] = {}
+        self.allowed_clients: dict[object, RsyncServer.Client] = {}
 
         self._shutdown_event = threading.Event()
         self._monitor_thread: Optional[threading.Thread] = None
diff --git a/Engine/Plugins/VirtualProduction/Switchboard/Source/Switchboard/switchboard/switchboard_dialog.py b/Engine/Plugins/VirtualProduction/Switchboard/Source/Switchboard/switchboard/switchboard_dialog.py
--- a/Engine/Plugins/VirtualProduction/Switchboard/Source/Switchboard/switchboard/switchboard_dialog.py
+++ b/Engine/Plugins/VirtualProduction/Switchboard/Source/Switchboard/switchboard/switchboard_dialog.py
@@ -1,5 +1,7 @@
 # Copyright Epic Games, Inc. All Rights Reserved.
 
+from __future__ import annotations
+
 import datetime
 import logging
 import os
@@ -154,7 +156,7 @@ class DeviceAdditionalSettingsUI(QtCore.QObject):
         self.settings_menu = QtWidgets.QMenu(parent)
         self._generate_settings_menu_items(self.settings_menu)
         self._button.setMenu(self.settings_menu)
-        self._button.setStyleSheet("QPushButton::menu-indicator{image:none;}");
+        self._button.setStyleSheet("QPushButton::menu-indicator{image:none;}")
         self._button.setDisabled(False)
 
     def make_button(self, parent):
@@ -163,32 +165,69 @@ class DeviceAdditionalSettingsUI(QtCore.QObject):
         """
         button = sb_widgets.ControlQPushButton.create(
                 icon_size=QtCore.QSize(21, 21),
-                tool_tip=f'Change device settings',
+                tool_tip='Change device settings',
                 hover_focus=False,
                 name='settings')
 
-        self.assign_button(button,parent)
+        self.assign_button(button, parent)
         return button
 
-class PeriodicRunnable(QtCore.QRunnable):
-    ''' Performs periodic tasks on the switchboard dialog. '''
 
-    def __init__(self, switchboard):
+class ProcessMonitor(QtCore.QObject):
+    ''' Offloads potentially expensive periodic checks to another thread. '''
+
+    mu_server_status_changed = QtCore.Signal(object)
+
+    class MuServerStatus:
+        def __init__(self):
+            self.pid: Optional[int] = None
+            self.endpoint: Optional[str] = None
+            self.name: Optional[str] = None
+
+        def __eq__(self, other):
+            if isinstance(other, ProcessMonitor.MuServerStatus):
+                return ((self.pid == other.pid) and
+                        (self.endpoint == other.endpoint) and
+                        (self.name == other.name))
+            else:
+                return False
+
+    def __init__(self, dialog: SwitchboardDialog):
         super().__init__()
-        self._switchboard = switchboard
+
+        self._dialog = dialog
         self._exiting = False
+        self._last_mu_status = ProcessMonitor.MuServerStatus()
+
+        self._thread = threading.Thread(target=self._run)
+        self._thread.start()
 
     def exit(self):
         self._exiting = True
+        self._thread.join()
 
-    def run(self):
+    def _run(self):
         while not self._exiting:
-            self._switchboard.update_muserver_button()
-            self._switchboard.update_locallistener_menuitem()
-            self._switchboard.update_insights_menuitem()
+            self._poll_multiuser_status()
+            self._dialog.update_locallistener_menuitem()
+            self._dialog.update_insights_menuitem()
 
             time.sleep(1.0)
 
+    def _poll_multiuser_status(self):
+        new_status = ProcessMonitor.MuServerStatus()
+
+        server = switchboard_application.get_multi_user_server_instance()
+        if server.is_running():
+            new_status.pid = server._running_pid
+            new_status.endpoint = server.running_endpoint()
+            new_status.name = server.running_server_name()
+
+        if new_status != self._last_mu_status:
+            self._last_mu_status = new_status
+            self.mu_server_status_changed.emit(new_status)
+
+
 class SwitchboardDialog(QtCore.QObject):
 
     STYLESHEET_PATH = os.path.join(RELATIVE_PATH, 'ui/switchboard.qss')
@@ -258,7 +297,6 @@ class SwitchboardDialog(QtCore.QObject):
         self._multiuser_session_name = None
         self._is_recording = False
         self._description = 'description'
-        self._started_mu_server = False
 
         # Recording Manager
         self.recording_manager = recording.RecordingManager(CONFIG.SWITCHBOARD_DIR)
@@ -476,59 +514,58 @@ class SwitchboardDialog(QtCore.QObject):
             SETTINGS.save()
             self.osc_server.launch(SETTINGS.ADDRESS.get_value(), CONFIG.OSC_SERVER_PORT.get_value())
 
-    def warn_user_about_muserver(self):
-        if self.have_warned_about_muserver:
-            return
-
-        ServerInstance = switchboard_application.get_multi_user_server_instance()
-        expected_server = ServerInstance.server_name()
-        expected_endpoint = ServerInstance.endpoint_address()
-        actual_server = ServerInstance.running_server_name()
-        actual_endpoint = ServerInstance.running_endpoint()
-        LOGGER.warning(f'The running Multi-user server does not match Switchboard configuration. Expected "{expected_server}" with "{expected_endpoint}" but found "{actual_server}" with "{actual_endpoint}". Please restart the multi-user server.')
-
-        self.have_warned_about_muserver = True
-
-    def get_muserver_label(self, is_running, is_valid):
-        ServerInstance = switchboard_application.get_multi_user_server_instance()
-        if is_running and is_valid:
-            actual_server = ServerInstance.running_server_name()
-            actual_endpoint = ServerInstance.running_endpoint()
-            return f'Multi-user Server ({actual_server} {actual_endpoint})'
-        elif is_running and not is_valid:
-            return 'Multi-user Server (WARNING: RUNNING SERVER DOES NOT MATCH CONFIGURATION)'
-        else:
-            return 'Multi-user Server'
-
-    def update_muserver_button(self):
+    def _on_mu_server_status_changed(
+            self,
+            new_status: ProcessMonitor.MuServerStatus
+    ):
         '''
         Update the status of the Multi-user start/stop button to reflect the status of Multi-user server.
         '''
-        ServerInstance = switchboard_application.get_multi_user_server_instance()
+        server = switchboard_application.get_multi_user_server_instance()
+        expected_name = server.configured_server_name()
+        expected_endpoint = server.configured_endpoint()
+
+        is_running = new_status.pid is not None
+        actual_name = new_status.name
+        actual_endpoint = new_status.endpoint
+
+        is_valid = (is_running and (actual_endpoint == expected_endpoint) and
+                    (actual_name == expected_name))
+
         is_checked = self.window.muserver_start_stop_button.isChecked()
-        is_running = ServerInstance.is_running()
-        is_valid =  ServerInstance.validate_process()
-        if  (is_running or self._started_mu_server) and not is_checked:
-            self.window.muserver_start_stop_button.setChecked( True )
+        if is_running and not is_checked:
+            self.window.muserver_start_stop_button.setChecked(True)
         elif not is_running and is_checked:
-            self.window.muserver_start_stop_button.setChecked( False )
+            self.window.muserver_start_stop_button.setChecked(False)
             self.have_warned_about_muserver = False
 
-        label = self.get_muserver_label(is_running or self._started_mu_server,is_valid)
-        self.window.muserver_label.setText(label)
-        if is_running and not is_valid:
-            self.warn_user_about_muserver()
+        # Re-enable the button after a pending launch completes
+        self.window.muserver_start_stop_button.setEnabled(True)
 
-        if is_running and self._started_mu_server:
-            # Server has finished starting so reset our start flag.
-            self._started_mu_server = False
+        label = 'Multi-user Server'
+        if is_running:
+            label += f' ({new_status.name} {new_status.endpoint})'
+            if not is_valid:
+                label += ' (WARNING: RUNNING SERVER DOES NOT MATCH CONFIGURATION)'
+                if not self.have_warned_about_muserver:
+                    LOGGER.warning(
+                        'The running Multi-user server does not match '
+                        'Switchboard configuration. Expected '
+                        f'"{expected_name}" with "{expected_endpoint}" but '
+                        f'found "{actual_name}" with "{actual_endpoint}". '
+                        'Please restart the multi-user server.')
+
+                    self.have_warned_about_muserver = True
+
+        self.window.muserver_label.setText(label)
 
     def setup_periodic_tasks_thread(self):
         '''
         Sets up a thread for performing maintenance tasks
         '''
-        self._periodic_runner = PeriodicRunnable(self)
-        QtCore.QThreadPool.globalInstance().start(self._periodic_runner)
+        self._periodic_runner = ProcessMonitor(self)
+        self._periodic_runner.mu_server_status_changed.connect(
+            self._on_mu_server_status_changed)
 
     def update_locallistener_menuitem(self):
         ''' 
@@ -551,12 +588,11 @@ class SwitchboardDialog(QtCore.QObject):
         '''
         ServerInstance = switchboard_application.get_multi_user_server_instance()
         if ServerInstance.is_running():
-            ServerInstance.terminate(bypolling=True)
-            self._started_mu_server = False
+            ServerInstance.process.kill()
         else:
-            self._started_mu_server = True
             ServerInstance.launch()
-        self.update_muserver_button()
+            # This will be re-enabled when the button transitions to stop/close
+            self.window.muserver_start_stop_button.setEnabled(False)
 
     def init_insights_launcher(self):
         ''' Initializes insights launcher '''
diff --git a/Engine/Plugins/VirtualProduction/Switchboard/Source/Switchboard/switchboard/switchboard_utils.py b/Engine/Plugins/VirtualProduction/Switchboard/Source/Switchboard/switchboard/switchboard_utils.py
--- a/Engine/Plugins/VirtualProduction/Switchboard/Source/Switchboard/switchboard/switchboard_utils.py
+++ b/Engine/Plugins/VirtualProduction/Switchboard/Source/Switchboard/switchboard/switchboard_utils.py
@@ -112,10 +112,11 @@ class PollProcess(object):
 
     @property
     def args(self) -> str:
-        get_commandline_cmd = f"wmic process where caption=\"{self.task_name}\" get commandline"
+        get_commandline_cmd = f"wmic process where caption=\"{self.task_name}\" get commandline /format:csv"
         try:
             output = subprocess.check_output(get_commandline_cmd, startupinfo=get_hidden_sp_startupinfo()).decode()
-            return output
+            dictobj = next(csv.DictReader(io.StringIO(output.strip())))
+            return dictobj['CommandLine'] if dictobj else ''
         except:
             return ""
 
diff --git a/Engine/Plugins/VirtualProduction/Switchboard/Source/Switchboard/tests/README.txt b/Engine/Plugins/VirtualProduction/Switchboard/Source/Switchboard/tests/README.txt
new file mode 100644
--- /dev/null
+++ b/Engine/Plugins/VirtualProduction/Switchboard/Source/Switchboard/tests/README.txt
@@ -0,0 +1,4 @@
+To run these tests:
+    > cd ...\Engine\Plugins\VirtualProduction\Switchboard\Source\Switchboard
+    > ..\..\..\..\..\Extras\ThirdPartyNotUE\SwitchboardThirdParty\Python\Scripts\activate
+    (switchboard_venv) > python -m unittest discover -v -s tests\
diff --git a/Engine/Plugins/VirtualProduction/Switchboard/Source/Switchboard/tests/test_command_line_extractor.py b/Engine/Plugins/VirtualProduction/Switchboard/Source/Switchboard/tests/test_command_line_extractor.py
new file mode 100644
--- /dev/null
+++ b/Engine/Plugins/VirtualProduction/Switchboard/Source/Switchboard/tests/test_command_line_extractor.py
@@ -0,0 +1,106 @@
+# Copyright Epic Games, Inc. All Rights Reserved.
+
+import sys
+from typing import Union
+import unittest
+from unittest.mock import patch
+
+from switchboard import switchboard_application
+
+
+class CommandLineExtractorTestCase(unittest.TestCase):
+    def _extractor(self, args: Union[str, list[str]]):
+        return switchboard_application.CommandLineExtractor(args)
+
+    def _tokens(self, args: Union[str, list[str]]):
+        return self._extractor(args).tokens
+
+    @property
+    def _mu_path(self):
+        if sys.platform.startswith('win'):
+            return (R'C:\Program Files\Epic Games\UE_5.2'
+                    R'\Engine\Binaries\Win64\UnrealMultiUserSlateServer.EXE')
+        else:
+            return ('[REDACTED_USER_HOME]/UE 5.2/'
+                    'Engine/Binaries/Linux/UnrealMultiUserSlateServer')
+
+    def _test_multiuser(self):
+        ''' Inner test, called with varying `sys.platform` mock patches. '''
+
+        self.assertEqual(
+            self._tokens(f'"{self._mu_path}"'),
+            [self._mu_path],
+            'argv[0] containing spaces and/or backslashes')
+
+        TRANSPORT_UNICAST = '127.0.0.1:9030'
+        SERVER_NAME = 'MU_Server-001'
+
+        EXPECTED_TOKENS = [
+            self._mu_path, f'-CONCERTSERVER={SERVER_NAME}',
+            f'-UDPMESSAGING_TRANSPORT_UNICAST={TRANSPORT_UNICAST}']
+
+        self.assertEqual(
+            self._tokens(
+                f'"{self._mu_path}"'
+                f' -CONCERTSERVER={SERVER_NAME}'
+                f' -UDPMESSAGING_TRANSPORT_UNICAST={TRANSPORT_UNICAST}'),
+            EXPECTED_TOKENS,
+            'argv[0] token should be unquoted')
+
+        self.assertEqual(
+            self._tokens(
+                f'"{self._mu_path}"'
+                f' -CONCERTSERVER="{SERVER_NAME}"'
+                f' -UDPMESSAGING_TRANSPORT_UNICAST="{TRANSPORT_UNICAST}"'),
+            EXPECTED_TOKENS,
+            'argv[1:] quoted value quotes removed')
+
+        self.assertEqual(
+            self._tokens(
+                f'"{self._mu_path}"'
+                f' "-CONCERTSERVER={SERVER_NAME}"'
+                f' "-UDPMESSAGING_TRANSPORT_UNICAST={TRANSPORT_UNICAST}"'),
+            EXPECTED_TOKENS,
+            'argv[1:] enclosing quotes removed')
+
+    def test_tokenization(self):
+        self.assertEqual(
+            self._tokens('foo'),
+            ['foo'],
+            'Base case')
+
+        self.assertEqual(
+            self._tokens('foo bar baz'),
+            ['foo', 'bar', 'baz'],
+            'Unquoted spaces split tokens')
+
+        self.assertEqual(
+            self._tokens('foo "bar baz"'),
+            ['foo', 'bar baz'],
+            'Quoted token containing space')
+
+    def test_get_value_for_switch(self):
+        extractor = self._extractor(
+            '"argv0"'
+            ' -foo=bar'
+            R' -baz="forty \"two\""')
+
+        self.assertEqual(
+            extractor.get_value_for_switch('foo'),
+            'bar')
+
+        self.assertEqual(
+            extractor.get_value_for_switch('baz'),
+            'forty "two"')
+
+    @patch('sys.platform', 'win32')
+    def test_multiuser_windows(self):
+        with self.subTest("Multi-user (sys.platform == 'win32')"):
+            self.assertTrue(self._mu_path.startswith('C:\\'))
+            self._test_multiuser()
+
+    @patch('sys.platform', 'linux')
+    def test_multiuser_linux(self):
+        with self.subTest("Multi-user (sys.platform == 'linux')"):
+            self.assertTrue(self._mu_path.startswith('[REDACTED_USER_HOME]
+            self._test_multiuser()
