Switchboard: Fix an issue where focusing then unfocusing a class setting resets all device overrides of that setting.

diff --git a/Engine/Plugins/VirtualProduction/Switchboard/Source/Switchboard/switchboard/config.py b/Engine/Plugins/VirtualProduction/Switchboard/Source/Switchboard/switchboard/config.py
--- a/Engine/Plugins/VirtualProduction/Switchboard/Source/Switchboard/switchboard/config.py
+++ b/Engine/Plugins/VirtualProduction/Switchboard/Source/Switchboard/switchboard/config.py
@@ -391,7 +391,7 @@ class Setting(QtCore.QObject):
 
     def _refresh_reset_override_widgets(self):
         for device_name in self._overrides.keys():
-            self._on_press_reset_override(device_name)
+            self._refresh_reset_override_widget(device_name)
 
     def _refresh_reset_override_widget(self, device_name: str):
         if device_name in self._reset_override_widgets:
