Switchboard: Add a means to force redeployment of remote Listeners.

- You can now access the Listener redeployer flow via the Tools menu, and check/uncheck individual devices at your discretion.
- Added an optional named argument to ListenerClient.connect to specify the socket timeout.

diff --git a/Engine/Plugins/VirtualProduction/Switchboard/Source/Switchboard/switchboard/devices/unreal/redeploy_dialog.py b/Engine/Plugins/VirtualProduction/Switchboard/Source/Switchboard/switchboard/devices/unreal/redeploy_dialog.py
--- a/Engine/Plugins/VirtualProduction/Switchboard/Source/Switchboard/switchboard/devices/unreal/redeploy_dialog.py
+++ b/Engine/Plugins/VirtualProduction/Switchboard/Source/Switchboard/switchboard/devices/unreal/redeploy_dialog.py
@@ -16,6 +16,13 @@ from switchboard.switchboard_logging import LOGGER
 from . import version_helpers
 
 
+COLOR_GREEN = '#4f4'
+COLOR_RED = '#f44'
+COLOR_YELLOW = '#ff4'
+
+CONNECT_TIMEOUT_SECONDS = 5
+
+
 class RedeployListenerEndpoint(QtCore.QObject):
     signal_refresh_ui = QtCore.Signal()
     signal_result = QtCore.Signal(bool, str, BaseException) # (bSuccess, details, exc_info)
@@ -64,38 +71,37 @@ class RedeployListenerEndpoint(QtCore.QObject):
         self.version_label = QtWidgets.QLabel('')
         self.version_label.setObjectName('version_label')
 
-    def can_redeploy(self):
-        if not self.version or not self.dlg_parent.listener_ver:
-            return False
+        self.redeploy_checkbox = QtWidgets.QCheckBox()
+        self.redeploy_checkbox.setObjectName('redeploy_checkbox')
 
-        redeploy_supported = version_helpers.listener_supports_redeploy(self.version)
-        compatible = version_helpers.listener_is_compatible(self.version)
-        old_version = self.version < self.dlg_parent.listener_ver
-
-        return redeploy_supported and (old_version or not compatible)
+    @QtCore.Slot()
+    def refresh_ui(self):
+        if not self.version:
+            return
 
-    def rich_version_str(self, remote_ver, available_redeploy_ver):
-        version_str = version_helpers.version_str(remote_ver)
-        if not version_helpers.listener_is_compatible(remote_ver) or not version_helpers.listener_supports_redeploy(remote_ver):
-            version_str = f'<span style="color: #f44">{version_str}</span>'
-        elif available_redeploy_ver and remote_ver < available_redeploy_ver:
-            version_str = f'<span style="color: #ff4">{version_str}</span>'
+        version_str = version_helpers.version_str(self.version)
+        if not version_helpers.listener_is_compatible(self.version):
+            version_color = COLOR_RED
+            self.redeploy_checkbox.setChecked(True)
+        elif (self.dlg_parent.listener_ver and
+              self.dlg_parent.listener_ver > self.version):
+            version_color = COLOR_YELLOW
+            self.redeploy_checkbox.setChecked(True)
         else:
-            version_str = f'<span style="color: #4f4">{version_str}</span>'
+            version_color = COLOR_GREEN
+            self.redeploy_checkbox.setChecked(False)
 
-        return version_str
+        rich_ver = f'<span style="color: {version_color}">{version_str}</span>'
 
-    @QtCore.Slot()
-    def refresh_ui(self):
-        if self.version:
-            self.version_label.setText(f"Version {self.rich_version_str(self.version, self.dlg_parent.listener_ver)}")
+        self.version_label.setText(f"Version {rich_ver}")
 
     @QtCore.Slot()
     def on_redeploy_clicked(self):
-        self.do_redeploy()
+        if self.redeploy_checkbox.isChecked():
+            self.do_redeploy()
 
     def do_redeploy(self):
-        if not self.can_redeploy():
+        if not version_helpers.listener_supports_redeploy(self.version):
             return
 
         if not self.dlg_parent.listener_exe_base64 or not self.dlg_parent.listener_exe_sha1sum:
@@ -115,7 +121,7 @@ class RedeployListenerEndpoint(QtCore.QObject):
         # Suppress on_disconnect for this expected disconnect
         self.client.disconnect()
         self.client.disconnect_delegate = self.on_disconnect
-        self.client.connect()
+        self.client.connect(timeout=CONNECT_TIMEOUT_SECONDS)
 
     def on_disconnect(self, unexpected, exception):
         self.version = None
@@ -159,25 +165,27 @@ class RedeployListenerDialog(QtWidgets.QDialog):
         self.listener_exe_base64: Optional[str] = None
         self.listener_exe_sha1sum: Optional[str] = None
 
-        self.setWindowTitle('Update Listeners')
+        self.setWindowTitle('Listener Redeployer')
 
         self.finished.connect(self.on_dialog_finished)
 
+        NUM_COLS = 4
+
         layout = QtWidgets.QGridLayout(self)
         self.setLayout(layout)
         layout.setHorizontalSpacing(20)
         row_idx = 0
 
         # Local version header row
-        layout.addWidget(QtWidgets.QLabel('<b>Local Listener</b>'), row_idx, 0, 1, 3)
+        layout.addWidget(QtWidgets.QLabel('<b>Local Listener</b>'), row_idx, 0, 1, NUM_COLS)
         row_idx += 1
 
         # Local listener info row
         self.listener_path_label = QtWidgets.QLabel('')
-        layout.addWidget(self.listener_path_label, row_idx, 0, 1, 2)
+        layout.addWidget(self.listener_path_label, row_idx, 0, 1, 3)
 
         self.listener_ver_label = QtWidgets.QLabel('')
-        layout.addWidget(self.listener_ver_label, row_idx, 2)
+        layout.addWidget(self.listener_ver_label, row_idx, 3)
 
         row_idx += 1
 
@@ -186,7 +194,7 @@ class RedeployListenerDialog(QtWidgets.QDialog):
         row_idx += 1
 
         # Listeners header row
-        layout.addWidget(QtWidgets.QLabel('<b>Remote Listeners</b>'), row_idx, 0, 1, 3)
+        layout.addWidget(QtWidgets.QLabel('<b>Remote Listeners</b>'), row_idx, 0, 1, NUM_COLS)
         row_idx += 1
 
         # Remote listener models and rows
@@ -196,10 +204,16 @@ class RedeployListenerDialog(QtWidgets.QDialog):
                 endpoint = RedeployListenerEndpoint(self, device.unreal_client.address, device.unreal_client.port)
                 endpoint.signal_result.connect(self.on_endpoint_result)
                 endpoint.signal_refresh_ui.connect(self.refresh_ui)
+                endpoint.redeploy_checkbox.stateChanged.connect(
+                    lambda: self.redeploy_btn.setEnabled(
+                        any(ep.redeploy_checkbox.isChecked()
+                            for ep in self.endpoints.values()))
+                )
                 endpoint.client.connect()
-                layout.addWidget(endpoint.endpoint_label, row_idx, 0)
-                layout.addWidget(endpoint.devices_label,  row_idx, 1)
-                layout.addWidget(endpoint.version_label,  row_idx, 2)
+                layout.addWidget(endpoint.redeploy_checkbox, row_idx, 0)
+                layout.addWidget(endpoint.endpoint_label,    row_idx, 1)
+                layout.addWidget(endpoint.devices_label,     row_idx, 2)
+                layout.addWidget(endpoint.version_label,     row_idx, 3)
                 self.endpoints[ep_addr] = endpoint
                 row_idx += 1
 
@@ -210,9 +224,10 @@ class RedeployListenerDialog(QtWidgets.QDialog):
         row_idx += 1
 
         # Button row
-        self.redeploy_all_btn = QtWidgets.QPushButton('Redeploy All')
-        self.redeploy_all_btn.clicked.connect(self.on_redeploy_all_clicked)
-        layout.addWidget(self.redeploy_all_btn, row_idx, 0, 1, 3)
+        self.redeploy_btn = QtWidgets.QPushButton('Redeploy Selected')
+        self.redeploy_btn.clicked.connect(self.on_redeploy_clicked)
+        self.redeploy_btn.setEnabled(False)
+        layout.addWidget(self.redeploy_btn, row_idx, 0, 1, NUM_COLS)
 
         for endpoint in self.endpoints.values():
             endpoint.devices_label.setText(f"{', '.join(d.name for d in endpoint.devices)}")
@@ -251,20 +266,19 @@ class RedeployListenerDialog(QtWidgets.QDialog):
     def refresh_ui(self):
         listener_path_label_str = f"<u>{os.path.split(self.listener_path)[1]}</u>"
         if not self.listener_ver or not self.listener_exe_base64 or not self.listener_exe_sha1sum:
-            listener_path_label_str = f'<span style="color: #f44">{listener_path_label_str}</span>'
+            listener_path_label_str = (f'<span style="color: {COLOR_RED}">'
+                                       f'{listener_path_label_str}</span>')
 
         self.listener_path_label.setText(listener_path_label_str)
         self.listener_path_label.setToolTip(self.listener_path)
 
         listener_ver_str = version_helpers.version_str(self.listener_ver) if self.listener_ver else '(N/A)'
         if not self.listener_ver or not version_helpers.listener_is_compatible(self.listener_ver):
-            listener_ver_str = f'<span style="color: #f44">{listener_ver_str}</span>'
+            listener_ver_str = (f'<span style="color: {COLOR_RED}">'
+                                f'{listener_ver_str}</span>')
 
         self.listener_ver_label.setText(f"Version {listener_ver_str}")
 
-        any_endpoints_can_redeploy = any(e.can_redeploy() for e in self.endpoints.values())
-        self.redeploy_all_btn.setEnabled(any_endpoints_can_redeploy)
-
     @QtCore.Slot()
     def on_listener_changed(self):
         self.try_refresh_listener()
@@ -275,7 +289,7 @@ class RedeployListenerDialog(QtWidgets.QDialog):
             endpoint.client.disconnect()
 
     @QtCore.Slot()
-    def on_redeploy_all_clicked(self):
+    def on_redeploy_clicked(self):
         for (ep_addr, endpoint) in self.endpoints.items():
             endpoint.on_redeploy_clicked()
 
diff --git a/Engine/Plugins/VirtualProduction/Switchboard/Source/Switchboard/switchboard/listener_client.py b/Engine/Plugins/VirtualProduction/Switchboard/Source/Switchboard/switchboard/listener_client.py
--- a/Engine/Plugins/VirtualProduction/Switchboard/Source/Switchboard/switchboard/listener_client.py
+++ b/Engine/Plugins/VirtualProduction/Switchboard/Source/Switchboard/switchboard/listener_client.py
@@ -100,7 +100,7 @@ class ListenerClient(object):
             return False
         return False
 
-    def connect(self, address=None):
+    def connect(self, address=None, *, timeout: Optional[float] = None):
         '''
         Initiates a connection.
 
@@ -122,16 +122,19 @@ class ListenerClient(object):
         self.last_activity = datetime.datetime.now()
 
         self.listener_qt_handler.listener_connecting.emit(self)
-        establish_connection_thread = Thread(target=self._establish_connection)
+        establish_connection_thread = Thread(target=self._establish_connection,
+                                             kwargs={'timeout': timeout})
         establish_connection_thread.start()
 
         return True
 
-    def _establish_connection(self):
+    def _establish_connection(self, *, timeout: Optional[float] = None):
         try:
             LOGGER.info(f"Connecting to {self.address}:{self.port}")
 
             self.socket = socket.socket(socket.AF_INET, socket.SOCK_STREAM)
+            if timeout is not None:
+                self.socket.settimeout(timeout)
             self.socket.connect(self.server_address)
 
             # Create a thread that waits for messages from the server
diff --git a/Engine/Plugins/VirtualProduction/Switchboard/Source/Switchboard/switchboard/switchboard_dialog.py b/Engine/Plugins/VirtualProduction/Switchboard/Source/Switchboard/switchboard/switchboard_dialog.py
--- a/Engine/Plugins/VirtualProduction/Switchboard/Source/Switchboard/switchboard/switchboard_dialog.py
+++ b/Engine/Plugins/VirtualProduction/Switchboard/Source/Switchboard/switchboard/switchboard_dialog.py
@@ -39,6 +39,7 @@ from switchboard.switchboard_logging import ConsoleStream, LOGGER
 from switchboard.tools.insights_launcher import InsightsLauncher
 from switchboard.tools.listener_launcher import ListenerLauncher
 from switchboard.devices.unreal.plugin_unreal import DeviceUnreal
+from switchboard.devices.unreal.redeploy_dialog import RedeployListenerDialog
 from switchboard.util import collect_logs
 
 ENGINE_PATH = "../../../../.."
@@ -316,6 +317,9 @@ class SwitchboardDialog(QtCore.QObject):
         self.register_open_logs_menuitem()
         self.register_zip_logs_menuitem()
 
+        # Redeploy menu item
+        self.register_redeploy_menuitem()
+
         # Transport Manager
         #self.transport_queue = recording.TransportQueue(CONFIG.SWITCHBOARD_DIR)
         #self.transport_queue.signal_transport_queue_job_started.connect(self.transport_queue_job_started)
@@ -641,6 +645,15 @@ class SwitchboardDialog(QtCore.QObject):
         action = self.register_tools_menu_action("&Zip Logs")
         action.triggered.connect(save_logs)
 
+    def register_redeploy_menuitem(self):
+        def show_redeploy_dialog():
+            dlg = RedeployListenerDialog(DeviceUnreal.active_unreal_devices,
+                                         DeviceUnreal.listener_watcher)
+            dlg.exec()
+
+        action = self.register_tools_menu_action("Listener &Redeployer")
+        action.triggered.connect(show_redeploy_dialog)
+
     def add_tools_menu(self):
         ''' Adds tools menu to menu bar and populates built-in items '''
 
