Switchboard: Fix syncing exception when "Engine Sync" is set to "Use Existing" and AdditionalPathsToSync are present.

diff --git a/Engine/Plugins/VirtualProduction/Switchboard/Source/Switchboard/sbl_helper.py b/Engine/Plugins/VirtualProduction/Switchboard/Source/Switchboard/sbl_helper.py
--- a/Engine/Plugins/VirtualProduction/Switchboard/Source/Switchboard/sbl_helper.py
+++ b/Engine/Plugins/VirtualProduction/Switchboard/Source/Switchboard/sbl_helper.py
@@ -657,8 +657,7 @@ class SbListenerHelper:
             project_paths.append(f'{self.uproj_path.parent / "..."}')
 
             for path in self.additional_paths_to_sync:
-                path = path.lstrip('/')
-                project_paths.append(f'{self.engine_dir.parent / path}')
+                project_paths.append(f'//{self.p4client}{path}')
 
             pathspecs.extend(
                 [f'{path}@{self.sync_project_cl}' for path in project_paths])
