Switchboard: Fix issue where sometimes ndisplay devices would not be populated in the dropdown.

The issue happens when there are .ndisplay files in the content folders. These are found but not assigned a class and the list of classes is too short, causing an exception later on and aborting populating the dropdown. Now .ndisplay items also get an ndisplay class assigned in the list.

diff --git a/Engine/Plugins/VirtualProduction/Switchboard/Source/Switchboard/switchboard/devices/unreal/plugin_unreal.py b/Engine/Plugins/VirtualProduction/Switchboard/Source/Switchboard/switchboard/devices/unreal/plugin_unreal.py
--- a/Engine/Plugins/VirtualProduction/Switchboard/Source/Switchboard/switchboard/devices/unreal/plugin_unreal.py
+++ b/Engine/Plugins/VirtualProduction/Switchboard/Source/Switchboard/switchboard/devices/unreal/plugin_unreal.py
@@ -2773,6 +2773,7 @@ class DeviceUnreal(Device):
                             asset_names.append(file_name)
                             asset_paths.append(asset_path)
                             asset_plugins.append(unreal_content_plugin)
+                            asset_classnames.append(DeviceUnreal.NDISPLAY_CLASS_NAMES[0]) # so that it passes the filter later on
 
         # process the assets in a multi-threaded fashion
 
@@ -2827,10 +2828,19 @@ class DeviceUnreal(Device):
                 # paths.
                 try:
                     asset = future.result()
-                    asset_names.append(asset['name'])
-                    asset_paths.append(asset['path'])
-                    asset_plugins.append(asset['plugin'])
-                    asset_classnames.append(asset['assetdata'].ObjectClassName)
+
+                    # first make sure they all exist
+                    name = asset['name']
+                    path = asset['path']
+                    plugin = asset['plugin']
+                    classname = asset['assetdata'].ObjectClassName
+
+                    # now append to lists
+                    asset_names.append(name)
+                    asset_paths.append(path)
+                    asset_plugins.append(plugin)
+                    asset_classnames.append(classname)
+                    
                 except Exception:
                     pass
 
