Switchboard: Make "built engine changelist" heuristic less likely to mismatch.

This is not a comprehensive fix, since code changes within the project / AdditionalPathsToSync will still be considered, but those are much less common.

diff --git a/Engine/Plugins/VirtualProduction/Switchboard/Source/Switchboard/switchboard/devices/unreal/plugin_unreal.py b/Engine/Plugins/VirtualProduction/Switchboard/Source/Switchboard/switchboard/devices/unreal/plugin_unreal.py
--- a/Engine/Plugins/VirtualProduction/Switchboard/Source/Switchboard/switchboard/devices/unreal/plugin_unreal.py
+++ b/Engine/Plugins/VirtualProduction/Switchboard/Source/Switchboard/switchboard/devices/unreal/plugin_unreal.py
@@ -2357,14 +2357,14 @@ class DeviceUnreal(Device):
             LOGGER.error(
                 f"{self.name}: Please change the roles for this device in "
                 "the settings or in the unreal project settings!")
-                
+
     def _on_receive_editor_version(self, content):
         '''
         Receives the Engine/Binaries/[platform]/UnrealEditor.version file _request_unreal_editor_version_file
         '''
         decoded_content = base64.b64decode(content).decode()
         data = json.loads(decoded_content)
-        self.built_engine_changelist = data.get("Changelist", None)
+        self.built_engine_changelist = data.get("CompatibleChangelist", None)
 
     def on_file_receive_failed(self, source_path, error):
         if source_path.endswith(DeviceUnreal.csettings["roles_filename"].get_value(self.name)):
