Switchboard: Fixes for colocated devices (those sharing the same Perforce workspace / working directory on the same machine), as used in nDisplay multi-process mGPU configurations.

Avoid initiating conflicting sync/build operations, and try to enforce that related device states are otherwise kept consistent.

diff --git a/Engine/Plugins/VirtualProduction/Switchboard/Source/Switchboard/switchboard/devices/unreal/plugin_unreal.py b/Engine/Plugins/VirtualProduction/Switchboard/Source/Switchboard/switchboard/devices/unreal/plugin_unreal.py
--- a/Engine/Plugins/VirtualProduction/Switchboard/Source/Switchboard/switchboard/devices/unreal/plugin_unreal.py
+++ b/Engine/Plugins/VirtualProduction/Switchboard/Source/Switchboard/switchboard/devices/unreal/plugin_unreal.py
@@ -16,7 +16,7 @@ import re
 import socket
 import sys
 import threading
-from typing import Callable, Optional, Union
+from typing import Callable, Generator, Optional, Union
 import uuid
 
 from PySide2 import QtCore, QtGui, QtWidgets
@@ -1464,6 +1464,13 @@ class DeviceUnreal(Device):
                 "There is nothing to sync!")
             return
 
+        for device in self.devices_sharing_workspace():
+            if device.status >= DeviceStatus.OPEN:
+                LOGGER.warning(f'{self.name}: Unable to sync, because another '
+                               f'device ({device.name}) using the same '
+                               'workspace is already running.')
+                return
+
         program_name = 'sync'
 
         # check if it is already on its way:
@@ -1597,13 +1604,21 @@ class DeviceUnreal(Device):
             unreal_client=self.unreal_client,
         )
 
-        if self.status != DeviceStatus.SYNCING:
-            self.status = DeviceStatus.SYNCING
+        for device in self.devices_sharing_workspace():
+            if device.status != DeviceStatus.SYNCING:
+                device.status = DeviceStatus.SYNCING
 
     def build(self):
         if self.exclude_from_build.get_value():
             return
 
+        for device in self.devices_sharing_workspace():
+            if device.status >= DeviceStatus.OPEN:
+                LOGGER.warning(f'{self.name}: Unable to build, because '
+                               f'another device ({device.name}) using the '
+                               'same workspace is already running.')
+                return
+
         program_name = 'build_project'
 
         # check if it is already on its way:
@@ -1732,7 +1747,9 @@ class DeviceUnreal(Device):
             LOGGER.info(
                 f"{self.name}: Sending {program_name} command: "
                 f"{ubt_path} {ubt_args}")
-            self.status = DeviceStatus.BUILDING
+            
+            for device in self.devices_sharing_workspace():
+                device.status = DeviceStatus.BUILDING
 
         # Queue the build command
         self.program_start_queue.add(
@@ -2108,9 +2125,11 @@ class DeviceUnreal(Device):
             self.status = DeviceStatus.OPEN
             self.unreal_started_signal.emit()
         elif prog.name.startswith('build_'):
-            self.status = DeviceStatus.BUILDING
+            for device in self.devices_sharing_workspace():
+                device.status = DeviceStatus.BUILDING
         elif prog.name == 'sync':
-            self.status = DeviceStatus.SYNCING
+            for device in self.devices_sharing_workspace():
+                device.status = DeviceStatus.SYNCING
 
         self.program_start_queue.update_running_program(prog=prog)
 
@@ -2130,10 +2149,11 @@ class DeviceUnreal(Device):
             LOGGER.error(f"Could not start {program_name}: {message['error']}")
 
             if program_name == 'sync' or program_name.startswith('build_'):
-                self.status = DeviceStatus.CLOSED
-                # Force to show existing project_changelist to hide
-                # building/syncing.
-                self.project_changelist = self.project_changelist
+                for device in self.devices_sharing_workspace():
+                    device.status = DeviceStatus.CLOSED
+
+                    # This has the effect of hiding building/syncing status.
+                    device.project_changelist = device.project_changelist
             elif program_name == 'unreal':
                 self.status = DeviceStatus.CLOSED
             elif program_name == 'retrieve':
@@ -2260,6 +2280,14 @@ class DeviceUnreal(Device):
 
             self.status = DeviceStatus.CLOSED
 
+            for other in self.devices_sharing_workspace(skip=self):
+                sync_method = CONFIG.ENGINE_SYNC_METHOD.get_value()
+                if sync_method == EngineSyncMethod.Build_Engine.value:
+                    other.engine_changelist = self.engine_changelist
+
+                other.project_changelist = self.project_changelist
+                other.status = DeviceStatus.CLOSED
+
         elif program_name.startswith('build_'):
             if returncode == 0:
                 LOGGER.info(f"{self.name}: {program_name} successful!")
@@ -2274,11 +2302,12 @@ class DeviceUnreal(Device):
                         LOGGER.error(f"{self.name}: {line}")
 
             if 'build_project' == program_name:
-                self.status = DeviceStatus.CLOSED
-                
-                self._request_project_changelist_number()
+                for device in self.devices_sharing_workspace():
+                    device.status = DeviceStatus.CLOSED
+
                 # Forces an update to the changelist field (to hide the
                 # Building state).
+                self._request_project_changelist_number()
                 if CONFIG.ENGINE_SYNC_METHOD.get_value() == EngineSyncMethod.Build_Engine.value:
                     self._request_engine_changelist_number()
                     self._request_unreal_editor_version_file()
@@ -2310,14 +2339,16 @@ class DeviceUnreal(Device):
                 LOGGER.info(
                     f"{self.name}: Project {project_name} "
                     f"is on revision {current_changelist}")
-                self.project_changelist = current_changelist
+                for device in self.devices_sharing_workspace():
+                    device.project_changelist = current_changelist
             elif program_name.endswith("engine"):
                 project_name = os.path.basename(
                     os.path.dirname(CONFIG.UPROJECT_PATH.get_value(self.name)))
                 LOGGER.info(
                     f"{self.name}: Engine used for project "
                     f"{project_name} is on revision {current_changelist}")
-                self.engine_changelist = current_changelist
+                for device in self.devices_sharing_workspace():
+                    device.engine_changelist = current_changelist
 
     def on_program_killed(self, message):
         '''
@@ -2364,7 +2395,9 @@ class DeviceUnreal(Device):
         '''
         decoded_content = base64.b64decode(content).decode()
         data = json.loads(decoded_content)
-        self.built_engine_changelist = data.get("CompatibleChangelist", None)
+        compatible_cl = data.get("CompatibleChangelist", None)
+        for device in self.devices_sharing_workspace():
+            device.built_engine_changelist = compatible_cl
 
     def on_file_receive_failed(self, source_path, error):
         if source_path.endswith(DeviceUnreal.csettings["roles_filename"].get_value(self.name)):
@@ -2400,34 +2433,74 @@ class DeviceUnreal(Device):
         #
         if process['name'].startswith('build_'):
             for line in lines:
-                if '@progress' in line:
-                    stepparts = line.split("'")
+                if '@progress' not in line:
+                    continue
+
+                stepparts = line.split("'")
 
-                    if len(stepparts) < 2:
-                        break
+                if len(stepparts) < 2:
+                    break
 
-                    step = stepparts[-2].strip()
+                step = stepparts[-2].strip()
+                percent = line.split(' ')[-1].strip()
 
-                    percent = line.split(' ')[-1].strip()
+                if percent[-1] != '%':
+                    continue
 
-                    if '%' == percent[-1]:
-                        self.device_qt_handler.signal_device_build_update.emit(
-                            self, step, percent
-                        )
-                        
+                for device in self.devices_sharing_workspace():
+                    device.device_qt_handler.signal_device_build_update.emit(
+                            device, step, percent)
 
         elif process['name'] == 'sync':
             for line in lines:
-                if 'Progress:' in line:
-                    match = re.search(r'Progress: (\d{1,3}\.\d\d%)', line)
-                    if match:
-                        sync_progress = match.group(1)
-                        self.device_qt_handler.signal_device_sync_update.emit(
-                            self, sync_progress)
+                if 'Progress:' not in line:
+                    continue
+
+                match = re.search(r'Progress: (\d{1,3}\.\d\d%)', line)
+                if not match:
+                    continue
+
+                sync_progress = match.group(1)
+                for device in self.devices_sharing_workspace():
+                    device.device_qt_handler.signal_device_sync_update.emit(
+                        device, sync_progress)
 
         for line in lines:
             LOGGER.debug(f"{self.name} {process['name']}: {line}")
 
+    @classmethod
+    def devices(
+        cls, *,
+        where: Callable[[DeviceUnreal], bool] = lambda _: True,
+        only_connected: bool = True,
+        skip: Optional[DeviceUnreal] = None,
+    ) -> Generator[DeviceUnreal, None, None]:
+        for other in cls.active_unreal_devices:
+            if other == skip:
+                continue
+
+            if only_connected and other.is_disconnected:
+                continue
+
+            if where(other):
+                yield other
+
+    def devices_sharing_workspace(
+        self, *,
+        only_connected: bool = True,
+        skip: Optional[DeviceUnreal] = None,
+    ) -> Generator[DeviceUnreal, None, None]:
+        def get_ws(device):
+            return CONFIG.SOURCE_CONTROL_WORKSPACE.get_value(
+                device.name).casefold()
+
+        self_ws = get_ws(self)
+        return DeviceUnreal.devices(
+            where=lambda x: get_ws(x) == self_ws,
+            only_connected=only_connected,
+            skip=skip,
+        )
+
     def on_listener_state(self, message):
         '''
         Message expected to be received upon connection with the listener.
diff --git a/Engine/Plugins/VirtualProduction/Switchboard/Source/Switchboard/switchboard/switchboard_dialog.py b/Engine/Plugins/VirtualProduction/Switchboard/Source/Switchboard/switchboard/switchboard_dialog.py
--- a/Engine/Plugins/VirtualProduction/Switchboard/Source/Switchboard/switchboard/switchboard_dialog.py
+++ b/Engine/Plugins/VirtualProduction/Switchboard/Source/Switchboard/switchboard/switchboard_dialog.py
@@ -8,8 +8,7 @@ import os
 import threading
 import time
 import re
-import shutil
-from typing import List, Optional, Set, Union
+from typing import Callable, List, Optional, Set, Union
 
 from pathlib import Path
 
@@ -32,7 +31,7 @@ from switchboard import switchboard_widgets as sb_widgets
 from switchboard.add_config_dialog import AddConfigDialog
 from switchboard.config import CONFIG, DEFAULT_MAP_TEXT, ENABLE_UGS_SUPPORT, SETTINGS, EngineSyncMethod
 from switchboard.device_list_widget import DeviceListWidget, DeviceWidgetHeader
-from switchboard.devices.device_base import DeviceStatus
+from switchboard.devices.device_base import Device, DeviceStatus
 from switchboard.devices.device_manager import DeviceManager
 from switchboard.settings_dialog import SettingsDialog
 from switchboard.switchboard_logging import ConsoleStream, LOGGER
@@ -1100,28 +1099,46 @@ class SwitchboardDialog(QtCore.QObject):
     def sync_all_button_clicked(self):
         if not CONFIG.P4_ENABLED.get_value():
             return
-        device_widgets = self.device_list_widget.device_widgets()
 
-        for device_widget in device_widgets:
-            if device_widget.can_sync():
-                device_widget.sync_button_clicked()
+        self._sync_build_all(do_sync=True, do_build=False)
 
     def build_all_button_clicked(self):
-        device_widgets = self.device_list_widget.device_widgets()
-
-        for device_widget in device_widgets:
-            if device_widget.can_build():
-                device_widget.build_button_clicked()
+        self._sync_build_all(do_sync=False, do_build=True)
 
     def sync_and_build_all_button_clicked(self):
         if not CONFIG.P4_ENABLED.get_value():
             return
-        device_widgets = self.device_list_widget.device_widgets()
+        
+        self._sync_build_all(do_sync=True, do_build=True)
+
+    def _sync_build_all(self, *, do_sync: bool, do_build: bool):
+        actions_str = ' and '.join(filter(
+            None, ['sync' if do_sync else '', 'build' if do_build else '']))
+
+        workspace_to_device_map: dict[str, Device] = {}
+
+        for device in self.device_list_widget.devices():
+            can_sync = device.widget.can_sync()
+            can_build = device.widget.can_build()
+
+            if (do_sync and not can_sync) or (do_build and not can_build):
+                continue
+
+            workspace = CONFIG.SOURCE_CONTROL_WORKSPACE.get_value(
+                device.name).casefold()
+            if existing_device := workspace_to_device_map.get(workspace):
+                LOGGER.info(f'Skipping {actions_str} for {device.name} '
+                            f'because {existing_device.name} is already '
+                            'queued using the same workspace')
+                continue
+
+            workspace_to_device_map[workspace] = device
+
+            if do_sync:
+                device.widget.sync_button_clicked()
 
-        for device_widget in device_widgets:
-            if device_widget.can_sync() and device_widget.can_build():
-                device_widget.sync_button_clicked()
-                device_widget.build_button_clicked()
+            if do_build:
+                device.widget.build_button_clicked()
 
     def refresh_project_cl_button_clicked(self):
         self.p4_refresh_project_cl()
