Switchboard: Fix regression where devices with unset workspace are unable to build.

Also changes the `devices_sharing_workspace()` helper not to fan out to other devices with empty workspaces.

diff --git a/Engine/Plugins/VirtualProduction/Switchboard/Source/Switchboard/switchboard/devices/unreal/plugin_unreal.py b/Engine/Plugins/VirtualProduction/Switchboard/Source/Switchboard/switchboard/devices/unreal/plugin_unreal.py
--- a/Engine/Plugins/VirtualProduction/Switchboard/Source/Switchboard/switchboard/devices/unreal/plugin_unreal.py
+++ b/Engine/Plugins/VirtualProduction/Switchboard/Source/Switchboard/switchboard/devices/unreal/plugin_unreal.py
@@ -2491,12 +2491,18 @@ class DeviceUnreal(Device):
         skip: Optional[DeviceUnreal] = None,
     ) -> Generator[DeviceUnreal, None, None]:
         def get_ws(device):
-            return CONFIG.SOURCE_CONTROL_WORKSPACE.get_value(
-                device.name).casefold()
+            return (CONFIG.SOURCE_CONTROL_WORKSPACE.get_value(
+                device.name) or '').casefold()
 
         self_ws = get_ws(self)
+        def predicate(device: DeviceUnreal):
+            if self_ws:
+                return get_ws(device) == self_ws
+            else:
+                return device is self
+
         return DeviceUnreal.devices(
-            where=lambda x: get_ws(x) == self_ws,
+            where=predicate,
             only_connected=only_connected,
             skip=skip,
         )
