Switchboard: Fix for exception when syncing nDisplay devices.

diff --git a/Engine/Plugins/VirtualProduction/Switchboard/Source/Switchboard/switchboard/devices/unreal/plugin_unreal.py b/Engine/Plugins/VirtualProduction/Switchboard/Source/Switchboard/switchboard/devices/unreal/plugin_unreal.py
--- a/Engine/Plugins/VirtualProduction/Switchboard/Source/Switchboard/switchboard/devices/unreal/plugin_unreal.py
+++ b/Engine/Plugins/VirtualProduction/Switchboard/Source/Switchboard/switchboard/devices/unreal/plugin_unreal.py
@@ -1543,7 +1543,7 @@ class DeviceUnreal(Device):
         if generate_proj_files:
             sync_args += ' --generate'
 
-        if self.csettings['use_sync_filters'].get_value():
+        if DeviceUnreal.csettings['use_sync_filters'].get_value():
             include_str = ''
             custom_view_str = self.csettings['custom_sync_view'].get_value()
 
