Switchboard: Allow limited use of ugs_utils.SyncFilters class even when P4Python is unavailable.

diff --git a/Engine/Plugins/VirtualProduction/Switchboard/Source/Switchboard/switchboard/ugs_utils.py b/Engine/Plugins/VirtualProduction/Switchboard/Source/Switchboard/switchboard/ugs_utils.py
--- a/Engine/Plugins/VirtualProduction/Switchboard/Source/Switchboard/switchboard/ugs_utils.py
+++ b/Engine/Plugins/VirtualProduction/Switchboard/Source/Switchboard/switchboard/ugs_utils.py
@@ -566,6 +566,11 @@ class SyncFilters:
 
     @classmethod
     def supported(cls):
+        '''
+        If this method returns False, only basic parsing of configured
+        categories is available. Other important methods, such as
+        `includes_path` and `exclude_category` will raise an exception.
+        '''
         return P4PYTHON_AVAILABLE
 
     class CsvDialect(csv.excel):
@@ -583,9 +588,10 @@ class SyncFilters:
     def __init__(self):
         self.categories: dict[uuid.UUID, SyncFilters.Category] = {}
         self.excluded_categories: set[uuid.UUID] = set()
-        self.map = P4.Map()
 
-        self._reset_map()
+        if P4PYTHON_AVAILABLE:
+            self.map = P4.Map()
+            self._reset_map()
 
     def _reset_map(self):
         self.map.clear()
@@ -598,6 +604,9 @@ class SyncFilters:
             self._read(sfile)
 
     def exclude_category(self, category_id: uuid.UUID):
+        if not P4PYTHON_AVAILABLE:
+            raise RuntimeError('This method currently requires P4PYTHON')
+
         if category_id in self.excluded_categories:
             return  # already excluded
 
@@ -609,6 +618,9 @@ class SyncFilters:
         self._map_exclusion(category)
 
     def includes_path(self, path: str) -> bool:
+        if not P4PYTHON_AVAILABLE:
+            raise RuntimeError('This method currently requires P4PYTHON')
+
         return self.map.includes(path)
 
     def _read(self, file: TextIO):
