Fix typo - effect -> affect

#skipci
#rnx

diff --git a/Engine/Plugins/VirtualProduction/Switchboard/Source/Switchboard/switchboard/devices/ndisplay/plugin_ndisplay.py b/Engine/Plugins/VirtualProduction/Switchboard/Source/Switchboard/switchboard/devices/ndisplay/plugin_ndisplay.py
--- a/Engine/Plugins/VirtualProduction/Switchboard/Source/Switchboard/switchboard/devices/ndisplay/plugin_ndisplay.py
+++ b/Engine/Plugins/VirtualProduction/Switchboard/Source/Switchboard/switchboard/devices/ndisplay/plugin_ndisplay.py
@@ -896,7 +896,7 @@ class DevicenDisplay(DeviceUnreal):
         # Insights traces parameters
         if CONFIG.INSIGHTS_TRACE_ENABLE.get_value():
 
-            LOGGER.warning(f"Unreal Insight Tracing is enabled for '{self.name}'. This may effect Unreal Engine performance.")
+            LOGGER.warning(f"Unreal Insight Tracing is enabled for '{self.name}'. This may affect Unreal Engine performance.")
 
             remote_utrace_path = self.get_utrace_filepath()
 
