Set the item data for the default map label to match the expected key value when checking the default map on launch.

diff --git a/Engine/Plugins/VirtualProduction/Switchboard/Source/Switchboard/switchboard/switchboard_dialog.py b/Engine/Plugins/VirtualProduction/Switchboard/Source/Switchboard/switchboard/switchboard_dialog.py
--- a/Engine/Plugins/VirtualProduction/Switchboard/Source/Switchboard/switchboard/switchboard_dialog.py
+++ b/Engine/Plugins/VirtualProduction/Switchboard/Source/Switchboard/switchboard/switchboard_dialog.py
@@ -1974,7 +1974,7 @@ class SwitchboardDialog(QtCore.QObject):
             return
         
         # To disambiguate, show the full path name in the drop-down
-        level_combo_box.setItemData(0, "Default level")
+        level_combo_box.setItemData(0, DEFAULT_MAP_TEXT)
         level_combo_box.setItemData(0, "Default level", QtCore.Qt.ToolTipRole)
         for short_name_index in range(len(short_name_list)):
             short_name = short_name_list[short_name_index]
