Switchboard:

-Add option to prime shaders when launching nodes from Switchboard
 * Add context menu to nDisplay nodes and Tools Menu to be able to launch DDC Fill commandlet for the specified node or all nodes, current level or all levels

diff --git a/Engine/Plugins/VirtualProduction/Switchboard/Source/Switchboard/switchboard/devices/unreal/plugin_unreal.py b/Engine/Plugins/VirtualProduction/Switchboard/Source/Switchboard/switchboard/devices/unreal/plugin_unreal.py
--- a/Engine/Plugins/VirtualProduction/Switchboard/Source/Switchboard/switchboard/devices/unreal/plugin_unreal.py
+++ b/Engine/Plugins/VirtualProduction/Switchboard/Source/Switchboard/switchboard/devices/unreal/plugin_unreal.py
@@ -904,6 +904,13 @@ class DeviceUnreal(Device):
             tool_tip="List of roles for this device"
         )
 
+        self.setting_ddc_build_platforms = MultiOptionSetting(
+            attr_name="ddc_build_platforms",
+            nice_name="DDC Build Platforms",
+            value=["Windows", "WindowsEditor", "Linux", "LinuxEditor"],
+            tool_tip="List of platforms for which to build the DDC cache for this device"
+        )
+
         autojoin_mu_server = kwargs.get("autojoin_mu_server", True)
         self.autojoin_mu_server = BoolSetting(
             attr_name="autojoin_mu_server",
@@ -1171,6 +1178,7 @@ class DeviceUnreal(Device):
     def device_settings(self):
         return super().device_settings() + [
             self.setting_roles,
+            self.setting_ddc_build_platforms,
             self.autojoin_mu_server,
             self.last_launch_command,
             self.last_log_path,
@@ -2062,6 +2070,50 @@ class DeviceUnreal(Device):
             self.generate_unreal_exe_path(),
             self.generate_unreal_command_line_args(map_name))
 
+    def fill_derived_data_cache(self, current_level_only=False, program_name="unreal_ddc"):
+
+        possible_platforms = self.setting_ddc_build_platforms.get_value()
+        if (len(possible_platforms) == 0):
+            LOGGER.error("No platforms selected for fill_derived_data_cache, canceling...")
+            return 0
+
+        separator = '+'
+        platform_string = ''
+
+        for platform in possible_platforms:
+            if self.setting_ddc_build_platforms.get_value(platform):
+                platform_string += platform + separator
+        platform_string = platform_string.removesuffix(separator) 
+        
+        current_level = CONFIG.CURRENT_LEVEL
+        current_level_valid = current_level_only and current_level != None and current_level != DEFAULT_MAP_TEXT
+        map_name = ' Map=' + current_level if current_level_valid else ''
+
+        args = CONFIG.UPROJECT_PATH.get_value(self.name) + map_name + " -run=DerivedDataCache -TargetPlatform=" + platform_string + " -fill -DDC=CreateInstalledEnginePak"
+        LOGGER.info('Filling ddc with arguments: ' + args)
+
+        puuid, msg = message_protocol.create_start_process_message(
+            prog_path= self.generate_unreal_exe_path(),
+            prog_args=args,
+            prog_name=program_name,
+            caller=self.name,
+            update_clients_with_stdout=True,
+            priority_modifier=sb_utils.PriorityModifier.Normal.value,
+            lock_gpu_clock=False,
+        )
+
+        self.program_start_queue.add(
+            ProgramStartQueueItem(
+                name=program_name,
+                puuid_dependency=None,
+                puuid=puuid,
+                msg_to_unreal_client=msg,
+            ),
+            unreal_client=self.unreal_client,
+        )
+
+        return puuid
+
     def launch(self, map_name, program_name="unreal"):
         if map_name == DEFAULT_MAP_TEXT:
             map_name = ''
@@ -2995,6 +3047,7 @@ class DeviceWidgetUnreal(DeviceWidget):
     signal_open_last_log = QtCore.Signal(object)
     signal_open_last_trace = QtCore.Signal(object)
     signal_copy_last_launch_command = QtCore.Signal(object)
+    signal_device_widget_fill_ddc = QtCore.Signal(object, bool)
 
     def __init__(self, name, device_hash, address, icons, parent=None):
         self._autojoin_visible = True
@@ -3370,3 +3423,9 @@ class DeviceWidgetUnreal(DeviceWidget):
         cmenu.addAction("Open fetched trace", lambda: self.signal_open_last_trace.emit(self))
         cmenu.addAction("Copy last launch command", lambda: self.signal_copy_last_launch_command.emit(self))
 
+        # Only create DDC submenu if node is connected
+        if self.connect_button.isChecked():
+            fill_ddc_menu = cmenu.addMenu("Fill DDC (Prepare Shaders)")
+        
+            current_level_action = fill_ddc_menu.addAction("Current Level", lambda: self.signal_device_widget_fill_ddc.emit(self, True))
+            all_levels_action = fill_ddc_menu.addAction("All Levels", lambda: self.signal_device_widget_fill_ddc.emit(self, False))
diff --git a/Engine/Plugins/VirtualProduction/Switchboard/Source/Switchboard/switchboard/switchboard_dialog.py b/Engine/Plugins/VirtualProduction/Switchboard/Source/Switchboard/switchboard/switchboard_dialog.py
--- a/Engine/Plugins/VirtualProduction/Switchboard/Source/Switchboard/switchboard/switchboard_dialog.py
+++ b/Engine/Plugins/VirtualProduction/Switchboard/Source/Switchboard/switchboard/switchboard_dialog.py
@@ -325,6 +325,9 @@ class SwitchboardDialog(QtCore.QObject):
         # Redeploy menu item
         self.register_redeploy_menuitem()
 
+        # Fill DDC submenu
+        self.register_fill_ddc_menuitem()
+
         # Transport Manager
         #self.transport_queue = recording.TransportQueue(CONFIG.SWITCHBOARD_DIR)
         #self.transport_queue.signal_transport_queue_job_started.connect(self.transport_queue_job_started)
@@ -685,6 +688,23 @@ class SwitchboardDialog(QtCore.QObject):
         action = self.register_tools_menu_action("Listener &Redeployer")
         action.triggered.connect(show_redeploy_dialog)
 
+    def register_fill_ddc_menuitem(self):
+        def fill_ddc(current_level_only):
+            for device in self.device_manager.devices():
+                if not device.is_disconnected:
+                    DeviceUnreal(device).fill_derived_data_cache(current_level_only)
+
+        def fill_ddc_current_level_action():
+            fill_ddc(True)
+
+        def fill_ddc_all_levels_action():
+            fill_ddc(False)
+
+        current_level_action = self.register_tools_menu_action("Current Level", ["Fill DDC (Prepare Shaders)"])
+        current_level_action.triggered.connect(fill_ddc_current_level_action)
+        all_levels_action = self.register_tools_menu_action("All Levels", ["Fill DDC (Prepare Shaders)"])
+        all_levels_action.triggered.connect(fill_ddc_all_levels_action)
+
     def add_tools_menu(self):
         ''' Adds tools menu to menu bar and populates built-in items '''
 
@@ -705,7 +725,7 @@ class SwitchboardDialog(QtCore.QObject):
 
         current_menu:QMenu = self.tools_menu
 
-        # iterate over menunames give
+        # iterate over menunames given
         for menuname in menunames:
 
             create_menu = True
@@ -1298,6 +1318,7 @@ class SwitchboardDialog(QtCore.QObject):
         device_widget.signal_device_widget_close.connect(self.device_widget_close)
         device_widget.signal_device_widget_sync.connect(self.device_widget_sync)
         device_widget.signal_device_widget_build.connect(self.device_widget_build)
+        device_widget.signal_device_widget_fill_ddc.connect(self.device_widget_fill_ddc)
         device_widget.signal_device_widget_trigger_start_toggled.connect(self.device_widget_trigger_start_toggled)
         device_widget.signal_device_widget_trigger_stop_toggled.connect(self.device_widget_trigger_stop_toggled)
 
@@ -1640,6 +1661,11 @@ class SwitchboardDialog(QtCore.QObject):
         device = self.device_manager.device_with_hash(device_widget.device_hash)
         device.build()
 
+    @QtCore.Slot(object)
+    def device_widget_fill_ddc(self, device_widget, current_level_only):
+        device = self.device_manager.device_with_hash(device_widget.device_hash)
+        device.fill_derived_data_cache(current_level_only)
+
     @QtCore.Slot(object)
     def device_widget_play(self, device_widget):
         device = self.device_manager.device_with_hash(device_widget.device_hash)
