Switchboard fixes

- Properly check to see if a device can fill_ddc
- Default set the autojoin_mu to False for base device class
- When searching for editors ensure that the found editor matches the known path in switchboard.

diff --git a/Engine/Plugins/VirtualProduction/Switchboard/Source/Switchboard/switchboard/add_config_dialog.py b/Engine/Plugins/VirtualProduction/Switchboard/Source/Switchboard/switchboard/add_config_dialog.py
--- a/Engine/Plugins/VirtualProduction/Switchboard/Source/Switchboard/switchboard/add_config_dialog.py
+++ b/Engine/Plugins/VirtualProduction/Switchboard/Source/Switchboard/switchboard/add_config_dialog.py
@@ -404,6 +404,13 @@ class AddConfigDialog(QtWidgets.QDialog):
         editorfolder = ''
         projectpath = ''
 
+        sbfolder = ''
+        try:
+            sbpath = os.path.abspath(__file__)
+            sbfolder = self.find_upstream_path_with_name(sbpath, "Engine")
+        except FileNotFoundError:
+            pass
+
         # pick the best candidate out of the detected ones
         for idx, editor in enumerate(editors):
             try:
@@ -411,19 +418,16 @@ class AddConfigDialog(QtWidgets.QDialog):
             except FileNotFoundError:
                 continue
 
-            projectpath = projects[idx]
+            if sbfolder == '' or editorfolder == sbfolder:
+                projectpath = projects[idx]
 
             # prefer when we have both editor and project paths
             if editorfolder and projectpath:
                 break
 
         # If process-based detection didn't work, try to find the Engine associated with the running Switchboard
-        if not editorfolder:
-            try:
-                sbpath = os.path.abspath(__file__)
-                editorfolder = self.find_upstream_path_with_name(sbpath, "Engine")
-            except FileNotFoundError:
-                pass
+        if not editorfolder and sbfolder:
+            editorfolder = sbfolder
 
         # Populate the fields that we're updating and were found
 
diff --git a/Engine/Plugins/VirtualProduction/Switchboard/Source/Switchboard/switchboard/devices/device_base.py b/Engine/Plugins/VirtualProduction/Switchboard/Source/Switchboard/switchboard/devices/device_base.py
--- a/Engine/Plugins/VirtualProduction/Switchboard/Source/Switchboard/switchboard/devices/device_base.py
+++ b/Engine/Plugins/VirtualProduction/Switchboard/Source/Switchboard/switchboard/devices/device_base.py
@@ -82,7 +82,7 @@ class Device(QtCore.QObject):
 
         self._name = name  # Assigned name
         self.device_qt_handler = DeviceQtHandler()
-
+        self.autojoin_mu_server = False
         self.setting_address = AddressSetting('address', 'Address', address)
         self.setting_address.signal_setting_changed.connect(
             lambda: self._check_recreate_osc_client())
diff --git a/Engine/Plugins/VirtualProduction/Switchboard/Source/Switchboard/switchboard/switchboard_dialog.py b/Engine/Plugins/VirtualProduction/Switchboard/Source/Switchboard/switchboard/switchboard_dialog.py
--- a/Engine/Plugins/VirtualProduction/Switchboard/Source/Switchboard/switchboard/switchboard_dialog.py
+++ b/Engine/Plugins/VirtualProduction/Switchboard/Source/Switchboard/switchboard/switchboard_dialog.py
@@ -1318,7 +1318,8 @@ class SwitchboardDialog(QtCore.QObject):
         device_widget.signal_device_widget_close.connect(self.device_widget_close)
         device_widget.signal_device_widget_sync.connect(self.device_widget_sync)
         device_widget.signal_device_widget_build.connect(self.device_widget_build)
-        device_widget.signal_device_widget_fill_ddc.connect(self.device_widget_fill_ddc)
+        if hasattr(device_widget, 'signal_device_widget_fill_ddc'):
+            device_widget.signal_device_widget_fill_ddc.connect(self.device_widget_fill_ddc)
         device_widget.signal_device_widget_trigger_start_toggled.connect(self.device_widget_trigger_start_toggled)
         device_widget.signal_device_widget_trigger_stop_toggled.connect(self.device_widget_trigger_stop_toggled)
 
