Switchboard: Remove accidentally introduced transitive dependency; sbl_helper.py should not require PySide2.

Moved the `map_name_is_valid` and `get_game_launch_level_path` helper methods alongside the constants they relate to in config.py, which is not a dependency of sbl_helper.py.

#rnx

diff --git a/Engine/Plugins/VirtualProduction/Switchboard/Source/Switchboard/switchboard/config.py b/Engine/Plugins/VirtualProduction/Switchboard/Source/Switchboard/switchboard/config.py
--- a/Engine/Plugins/VirtualProduction/Switchboard/Source/Switchboard/switchboard/config.py
+++ b/Engine/Plugins/VirtualProduction/Switchboard/Source/Switchboard/switchboard/config.py
@@ -20,7 +20,8 @@ from switchboard import switchboard_widgets as sb_widgets
 from switchboard.switchboard_logging import LOGGER
 from switchboard.switchboard_widgets import (
     DropDownMenuComboBox, NonScrollableComboBox)
-from switchboard import ue_plugin_utils
+from switchboard import ue_plugin_utils, ugs_utils
+
 
 ROOT_CONFIGS_PATH = pathlib.Path(__file__).parent.with_name('configs')
 CONFIG_SUFFIX = '.json'
@@ -35,6 +36,48 @@ DEFAULT_MAP_TEXT = '-- Default Map --'
 ENABLE_UGS_SUPPORT = False
 
 
+def map_name_is_valid(map_name: str) -> bool:
+    """
+    Check if map is specified
+    """
+    return (map_name != ''
+            and map_name is not None
+            and map_name != DEFAULT_MAP_TEXT
+            and map_name != 'Default level'
+            and map_name != 'None')
+
+
+def get_game_launch_level_path() -> str:
+    level = CONFIG.CURRENT_LEVEL
+
+    # check if map name is valid. If not - get map name from DefaultEngine.ini
+    if not map_name_is_valid(level):
+        level = ''
+
+        # check if DefaultEngine.ini exists
+        engine_ini = os.path.join(pathlib.Path(CONFIG.UPROJECT_PATH.get_value()).parent, 'Config', 'DefaultEngine.ini')
+        if not os.path.exists(engine_ini):
+            LOGGER.warning("DefaultEngine.ini not found")
+            return ''
+
+        parser = ugs_utils.IniParser()
+
+        try:
+            engine_ini_file = open(engine_ini, "r", encoding="utf-8")
+            parser.read_file(engine_ini_file)
+
+            levels = parser.try_get('/Script/EngineSettings.GameMapsSettings', 'GameDefaultMap')
+
+            if len(levels) == 0:
+                raise Exception('Default map is not set in DefaultEngine.ini')
+            
+            level = levels[0]
+        except:
+            LOGGER.warning("Map is not set. Please select 'Level' to run dedicated server.")
+
+    return str(level)
+
+
 def migrate_comma_separated_string_to_list(value) -> list[str]:
     if isinstance(value, str):
         return value.split(",")
@@ -2848,4 +2891,3 @@ def list_config_paths() -> list[pathlib.Path]:
 # Get the user settings and load their config
 SETTINGS = UserSettings()
 CONFIG = Config()
-    
diff --git a/Engine/Plugins/VirtualProduction/Switchboard/Source/Switchboard/switchboard/devices/mpserver/plugin_mpserver.py b/Engine/Plugins/VirtualProduction/Switchboard/Source/Switchboard/switchboard/devices/mpserver/plugin_mpserver.py
--- a/Engine/Plugins/VirtualProduction/Switchboard/Source/Switchboard/switchboard/devices/mpserver/plugin_mpserver.py
+++ b/Engine/Plugins/VirtualProduction/Switchboard/Source/Switchboard/switchboard/devices/mpserver/plugin_mpserver.py
@@ -1,8 +1,8 @@
 # Copyright Epic Games, Inc. All Rights Reserved.
 
 from switchboard.devices.unreal.plugin_unreal import DeviceUnreal, DeviceWidgetUnreal
-from switchboard.config import (CONFIG, FilePathSetting, StringSetting, IntSetting)
-from switchboard.switchboard_utils import get_game_launch_level_path
+from switchboard.config import (CONFIG, FilePathSetting, StringSetting,
+                                IntSetting, get_game_launch_level_path)
 
 
 class DeviceMultiplayerServer(DeviceUnreal):
diff --git a/Engine/Plugins/VirtualProduction/Switchboard/Source/Switchboard/switchboard/devices/ndisplay/plugin_ndisplay.py b/Engine/Plugins/VirtualProduction/Switchboard/Source/Switchboard/switchboard/devices/ndisplay/plugin_ndisplay.py
--- a/Engine/Plugins/VirtualProduction/Switchboard/Source/Switchboard/switchboard/devices/ndisplay/plugin_ndisplay.py
+++ b/Engine/Plugins/VirtualProduction/Switchboard/Source/Switchboard/switchboard/devices/ndisplay/plugin_ndisplay.py
@@ -19,7 +19,8 @@ from switchboard import switchboard_widgets as sb_widgets
 from switchboard import switchboard_dialog as sb_dialog
 from switchboard.config import CONFIG, BoolSetting, IntSetting, FilePathSetting, \
     LoggingSetting, OptionSetting, Setting, StringSetting, SETTINGS, \
-    StringListSetting, AddressSetting, migrate_comma_separated_string_to_list
+    StringListSetting, AddressSetting, migrate_comma_separated_string_to_list, \
+    map_name_is_valid, get_game_launch_level_path
 from switchboard.devices.device_widget_base import AddDeviceDialog
 from switchboard.devices.unreal.plugin_unreal import DeviceUnreal, \
     DeviceWidgetUnreal, LiveLinkPresetSetting, MediaProfileSetting
@@ -28,7 +29,6 @@ from switchboard.switchboard_logging import LOGGER
 
 from .ndisplay_monitor_ui import nDisplayMonitorUI
 from .ndisplay_monitor import nDisplayMonitor
-from ...switchboard_utils import map_name_is_valid, get_game_launch_level_path
 
 
 class AddnDisplayDialog(AddDeviceDialog):
diff --git a/Engine/Plugins/VirtualProduction/Switchboard/Source/Switchboard/switchboard/p4_utils.py b/Engine/Plugins/VirtualProduction/Switchboard/Source/Switchboard/switchboard/p4_utils.py
--- a/Engine/Plugins/VirtualProduction/Switchboard/Source/Switchboard/switchboard/p4_utils.py
+++ b/Engine/Plugins/VirtualProduction/Switchboard/Source/Switchboard/switchboard/p4_utils.py
@@ -8,6 +8,8 @@ import pathlib
 import subprocess
 from typing import Optional
 
+# NOTE: This file should not depend on Qt, directly or indirectly.
+
 from . import switchboard_utils as sb_utils
 from .devices.unreal.version_helpers import LISTENER_COMPATIBLE_VERSION
 
diff --git a/Engine/Plugins/VirtualProduction/Switchboard/Source/Switchboard/switchboard/switchboard_utils.py b/Engine/Plugins/VirtualProduction/Switchboard/Source/Switchboard/switchboard/switchboard_utils.py
--- a/Engine/Plugins/VirtualProduction/Switchboard/Source/Switchboard/switchboard/switchboard_utils.py
+++ b/Engine/Plugins/VirtualProduction/Switchboard/Source/Switchboard/switchboard/switchboard_utils.py
@@ -9,12 +9,8 @@ import subprocess
 import threading
 from typing import Optional
 
-from pathlib import Path
-from .config import CONFIG, DEFAULT_MAP_TEXT
-from .switchboard_logging import LOGGER
-from re import split
+# NOTE: This file should not depend on Qt, directly or indirectly.
 
-from switchboard import ugs_utils
 
 class PriorityModifier(Enum):
     ''' Corresponds to `int32 PriorityModifier` parameter to `FPlatformProcess::CreateProc`. '''
@@ -199,42 +195,3 @@ def expand_endpoint(
         port_str = str(default_port)
 
     return f'{addr_str}:{port_str}'
-
-
-def map_name_is_valid(map_name: str) -> bool:
-    """
-    Check if map is specified
-    """
-    return map_name != '' and map_name is not None and map_name != DEFAULT_MAP_TEXT \
-        and map_name != 'Default level' and map_name != 'None'
-
-
-def get_game_launch_level_path() -> str:
-    level = CONFIG.CURRENT_LEVEL
-
-    #check if map name is valid. If not - get map name from DefaultEngine.ini
-    if not map_name_is_valid(level):
-        level = ''
-
-        #check if DefaultEngine.ini exists
-        engine_ini = os.path.join(Path(CONFIG.UPROJECT_PATH.get_value()).parent, 'Config', 'DefaultEngine.ini')
-        if not os.path.exists(engine_ini):
-            LOGGER.warning("DefaultEngine.ini not found")
-            return ''
-
-        parser = ugs_utils.IniParser()
-
-        try:
-            engine_ini_file = open(engine_ini, "r", encoding="utf-8")
-            parser.read_file(engine_ini_file)
-
-            levels = parser.try_get('/Script/EngineSettings.GameMapsSettings', 'GameDefaultMap')
-
-            if len(levels) == 0:
-                raise Exception('Default map is not set in DefaultEngine.ini')
-            
-            level = levels[0]
-        except:
-            LOGGER.warning("Map is not set. Please select 'Level' to run dedicated server.")
-
-    return str(level)
\ No newline at end of file
