switchboard: fixed missing include

diff --git a/Engine/Plugins/VirtualProduction/Switchboard/Source/Switchboard/switchboard/devices/ndisplay/plugin_ndisplay.py b/Engine/Plugins/VirtualProduction/Switchboard/Source/Switchboard/switchboard/devices/ndisplay/plugin_ndisplay.py
--- a/Engine/Plugins/VirtualProduction/Switchboard/Source/Switchboard/switchboard/devices/ndisplay/plugin_ndisplay.py
+++ b/Engine/Plugins/VirtualProduction/Switchboard/Source/Switchboard/switchboard/devices/ndisplay/plugin_ndisplay.py
@@ -25,6 +25,7 @@ from switchboard.devices.device_widget_base import AddDeviceDialog
 from switchboard.devices.unreal.plugin_unreal import DeviceUnreal, \
     DeviceWidgetUnreal, LiveLinkPresetSetting, MediaProfileSetting
 from switchboard.devices.unreal.uassetparser import UassetParser
+from switchboard.devices.device_base import DeviceStatus
 from switchboard.switchboard_logging import LOGGER
 
 from .ndisplay_monitor_ui import nDisplayMonitorUI
@@ -713,8 +714,8 @@ class DevicenDisplay(DeviceUnreal):
 
         cluster_nodes_count = 0
         for device in self.active_unreal_devices:
-            is_device_connected = device.status == DeviceStatus.CLOSED or device.status == DeviceStatus.OPEN
-            if device.device_type == "nDisplay" and is_device_connected:
+            is_device_connected = (device.status == DeviceStatus.CLOSED) or (device.status == DeviceStatus.OPEN)
+            if (device.device_type == "nDisplay") and is_device_connected:
                 cluster_nodes_count = cluster_nodes_count + 1
 
         def get_common_args():
