Switchboard: Warn user when attempting to rename nDisplay devices, or if missing the node corresponding to a named device when launching.

diff --git a/Engine/Plugins/VirtualProduction/Switchboard/Source/Switchboard/switchboard/devices/ndisplay/plugin_ndisplay.py b/Engine/Plugins/VirtualProduction/Switchboard/Source/Switchboard/switchboard/devices/ndisplay/plugin_ndisplay.py
--- a/Engine/Plugins/VirtualProduction/Switchboard/Source/Switchboard/switchboard/devices/ndisplay/plugin_ndisplay.py
+++ b/Engine/Plugins/VirtualProduction/Switchboard/Source/Switchboard/switchboard/devices/ndisplay/plugin_ndisplay.py
@@ -1,6 +1,5 @@
 # Copyright Epic Games, Inc. All Rights Reserved.
 
-import concurrent.futures
 import json
 import os
 import sys
@@ -227,6 +226,16 @@ class DeviceWidgetnDisplay(DeviceWidgetUnreal):
 
     signal_device_widget_primary = QtCore.Signal(object)
 
+    def on_name_changed(self, text):
+        super().on_name_changed(text)
+
+        if self.name_line_edit.is_valid:
+            point = self.name_line_edit.parent().mapToGlobal(
+                self.name_line_edit.geometry().topRight())
+            self.help_tool_tip.showText(
+                point,
+                "WARNING: Device names must match the nDisplay configuration")
+
     def _add_control_buttons(self):
         self._autojoin_visible = False
         super()._add_control_buttons()
@@ -1209,8 +1218,9 @@ class DevicenDisplay(DeviceUnreal):
             dev for dev in cls.active_unreal_devices
             if dev.device_type == 'nDisplay']
 
-        for device in devices:
+        shown_unknown_node_warning = False
 
+        for device in devices:
             # We exclude devices by disconnecting from them
             if device.is_disconnected:
                 continue
@@ -1221,9 +1231,19 @@ class DevicenDisplay(DeviceUnreal):
             # No name match, no node. Renaming nodes in Switchboard is not
             # currently supported.
             if node is None:
-                LOGGER.warning(
+                LOGGER.error(
                     f'Skipped "{device.name}" because it did not match any '
                     'node in the nDisplay configuration')
+
+                if not shown_unknown_node_warning:
+                    shown_unknown_node_warning = True
+                    QtWidgets.QMessageBox.warning(
+                        None, 'Unknown nDisplay Node',
+                        'One or more nDisplay devices have names that do not '
+                        'correspond to nodes in the configuration. This may '
+                        'prevent the cluster from launching correctly. Refer '
+                        'to the log for more details.')
+
                 continue
 
             # override address
