Switchboard: Fixes for silent failure when attempting to fill DDC for all connected devices.

diff --git a/Engine/Plugins/VirtualProduction/Switchboard/Source/Switchboard/switchboard/devices/unreal/plugin_unreal.py b/Engine/Plugins/VirtualProduction/Switchboard/Source/Switchboard/switchboard/devices/unreal/plugin_unreal.py
--- a/Engine/Plugins/VirtualProduction/Switchboard/Source/Switchboard/switchboard/devices/unreal/plugin_unreal.py
+++ b/Engine/Plugins/VirtualProduction/Switchboard/Source/Switchboard/switchboard/devices/unreal/plugin_unreal.py
@@ -2071,20 +2071,17 @@ class DeviceUnreal(Device):
             self.generate_unreal_command_line_args(map_name))
 
     def fill_derived_data_cache(self, current_level_only=False, program_name="unreal_ddc"):
-
-        possible_platforms = self.setting_ddc_build_platforms.get_value()
-        if (len(possible_platforms) == 0):
+        platforms: list[str] = self.setting_ddc_build_platforms.get_value(self.name)
+        if (len(platforms) == 0):
             LOGGER.error("No platforms selected for fill_derived_data_cache, canceling...")
             return 0
 
         separator = '+'
         platform_string = ''
+        for platform in platforms:
+            platform_string += platform + separator
+        platform_string = platform_string.removesuffix(separator)
 
-        for platform in possible_platforms:
-            if self.setting_ddc_build_platforms.get_value(platform):
-                platform_string += platform + separator
-        platform_string = platform_string.removesuffix(separator) 
-        
         current_level = CONFIG.CURRENT_LEVEL
         current_level_valid = current_level_only and current_level != None and current_level != DEFAULT_MAP_TEXT
         map_name = ' Map=' + current_level if current_level_valid else ''
@@ -2093,7 +2090,7 @@ class DeviceUnreal(Device):
         LOGGER.info('Filling ddc with arguments: ' + args)
 
         puuid, msg = message_protocol.create_start_process_message(
-            prog_path= self.generate_unreal_exe_path(),
+            prog_path=self.generate_unreal_exe_path(),
             prog_args=args,
             prog_name=program_name,
             caller=self.name,
diff --git a/Engine/Plugins/VirtualProduction/Switchboard/Source/Switchboard/switchboard/switchboard_dialog.py b/Engine/Plugins/VirtualProduction/Switchboard/Source/Switchboard/switchboard/switchboard_dialog.py
--- a/Engine/Plugins/VirtualProduction/Switchboard/Source/Switchboard/switchboard/switchboard_dialog.py
+++ b/Engine/Plugins/VirtualProduction/Switchboard/Source/Switchboard/switchboard/switchboard_dialog.py
@@ -692,7 +692,7 @@ class SwitchboardDialog(QtCore.QObject):
         def fill_ddc(current_level_only):
             for device in self.device_manager.devices():
                 if not device.is_disconnected:
-                    DeviceUnreal(device).fill_derived_data_cache(current_level_only)
+                    device.fill_derived_data_cache(current_level_only)
 
         def fill_ddc_current_level_action():
             fill_ddc(True)
