Switchboard: Display "Reinstall Dependencies" option even if the user has not created desktop/start menu shortcuts.

Also updates virtual environment integrity check and setup wizard text.

diff --git a/Engine/Plugins/VirtualProduction/Switchboard/Source/Switchboard/sb_setup.py b/Engine/Plugins/VirtualProduction/Switchboard/Source/Switchboard/sb_setup.py
--- a/Engine/Plugins/VirtualProduction/Switchboard/Source/Switchboard/sb_setup.py
+++ b/Engine/Plugins/VirtualProduction/Switchboard/Source/Switchboard/sb_setup.py
@@ -221,6 +221,7 @@ class SbSetup:
     def run_verify(self, options: argparse.Namespace) -> Dict[str, Any]:
         result = {}
         result['imports'] = {
+            'aioquic': False,
             'PySide2': False,
             'pythonosc.osc_server': False,
             'requests': False,
diff --git a/Engine/Plugins/VirtualProduction/Switchboard/Source/SwitchboardEditor/Private/SwitchboardSettingsCustomization.cpp b/Engine/Plugins/VirtualProduction/Switchboard/Source/SwitchboardEditor/Private/SwitchboardSettingsCustomization.cpp
--- a/Engine/Plugins/VirtualProduction/Switchboard/Source/SwitchboardEditor/Private/SwitchboardSettingsCustomization.cpp
+++ b/Engine/Plugins/VirtualProduction/Switchboard/Source/SwitchboardEditor/Private/SwitchboardSettingsCustomization.cpp
@@ -33,9 +33,8 @@ FText FSwitchboardEditorSettingsCustomization::GetHealthRowHintText() const
 	switch (FSwitchboardEditorModule::Get().GetSwitchboardInstallState())
 	{
 		case FSwitchboardEditorModule::ESwitchboardInstallState::Nominal:
-			return LOCTEXT("ReinstallEnvHint", "Switchboard requires additional library/framework dependencies to launch, which are installed. If you encounter issues launching Switchboard, it may help to reinstall them.");
 		case FSwitchboardEditorModule::ESwitchboardInstallState::ShortcutsMissing:
-			return LOCTEXT("AddShortcutsHint", "Add shortcuts to launch Switchboard from the Desktop and Start Menu.");
+			return LOCTEXT("ReinstallEnvHint", "Switchboard requires additional library/framework dependencies to launch, which are installed. If you encounter issues launching Switchboard, it may help to reinstall them.");
 		case FSwitchboardEditorModule::ESwitchboardInstallState::NeedInstallOrRepair:
 		default:
 			return LOCTEXT("InstallEnvHint", "Switchboard requires additional library/framework dependencies to launch, please install them.");
@@ -48,9 +47,8 @@ FText FSwitchboardEditorSettingsCustomization::GetHealthRowButtonText() const
 	switch (FSwitchboardEditorModule::Get().GetSwitchboardInstallState())
 	{
 		case FSwitchboardEditorModule::ESwitchboardInstallState::Nominal:
-			return LOCTEXT("ReinstallEnvButton", "Reinstall Dependencies");
 		case FSwitchboardEditorModule::ESwitchboardInstallState::ShortcutsMissing:
-			return LOCTEXT("AddShortcutsButton", "Add Shortcuts");
+			return LOCTEXT("ReinstallEnvButton", "Reinstall Dependencies");
 		case FSwitchboardEditorModule::ESwitchboardInstallState::NeedInstallOrRepair:
 		default:
 			return LOCTEXT("InstallEnvButton", "Install Dependencies");
@@ -63,8 +61,8 @@ const FSlateBrush* FSwitchboardEditorSettingsCustomization::GetHealthRowBorderBr
 	switch (FSwitchboardEditorModule::Get().GetSwitchboardInstallState())
 	{
 		case FSwitchboardEditorModule::ESwitchboardInstallState::Nominal:
-			return FSwitchboardEditorStyle::Get().GetBrush("Settings.RowBorder.Nominal");
 		case FSwitchboardEditorModule::ESwitchboardInstallState::ShortcutsMissing:
+			return FSwitchboardEditorStyle::Get().GetBrush("Settings.RowBorder.Nominal");
 		case FSwitchboardEditorModule::ESwitchboardInstallState::NeedInstallOrRepair:
 		default:
 			return FSwitchboardEditorStyle::Get().GetBrush("Settings.RowBorder.Warning");
@@ -77,8 +75,8 @@ const FSlateBrush* FSwitchboardEditorSettingsCustomization::GetHealthRowIconBrus
 	switch (FSwitchboardEditorModule::Get().GetSwitchboardInstallState())
 	{
 		case FSwitchboardEditorModule::ESwitchboardInstallState::Nominal:
-			return FSwitchboardEditorStyle::Get().GetBrush("Settings.Icons.Nominal");
 		case FSwitchboardEditorModule::ESwitchboardInstallState::ShortcutsMissing:
+			return FSwitchboardEditorStyle::Get().GetBrush("Settings.Icons.Nominal");
 		case FSwitchboardEditorModule::ESwitchboardInstallState::NeedInstallOrRepair:
 		default:
 			return FSwitchboardEditorStyle::Get().GetBrush("Settings.Icons.Warning");
@@ -98,10 +96,11 @@ void FSwitchboardEditorSettingsCustomization::CustomizeDetails(IDetailLayoutBuil
 	TSharedPtr<IPropertyHandle> VenvProperty = DetailLayout.GetProperty(GET_MEMBER_NAME_CHECKED(USwitchboardEditorSettings, VirtualEnvironmentPath));
 	if (VenvProperty)
 	{
-		// Filter health row together with venv property row.
+		// Filter health rows together with venv property row.
 		HealthFilterString = VenvProperty->GetPropertyDisplayName();
 	}
 
+	// Install health row.
 	SwitchboardCategory.AddCustomRow(HealthFilterString, false)
 		.WholeRowWidget
 		[
@@ -134,18 +133,73 @@ void FSwitchboardEditorSettingsCustomization::CustomizeDetails(IDetailLayoutBuil
 					SNew(SButton)
 					.Text(TAttribute<FText>::CreateSP(this, &FSwitchboardEditorSettingsCustomization::GetHealthRowButtonText))
 					.OnClicked_Lambda([this]() {
-						SSwitchboardSetupWizard::EWizardPage StartPage = SSwitchboardSetupWizard::EWizardPage::Intro;
-						if (FSwitchboardEditorModule::Get().GetSwitchboardInstallState() == FSwitchboardEditorModule::ESwitchboardInstallState::ShortcutsMissing)
-						{
-							StartPage = SSwitchboardSetupWizard::EWizardPage::Shortcuts;
-						}
-						SSwitchboardSetupWizard::OpenWindow(StartPage);
+						SSwitchboardSetupWizard::OpenWindow(SSwitchboardSetupWizard::EWizardPage::Intro);
 						return FReply::Handled();
 					})
 				]
 			]
 		];
 
+	// Add shortcuts row.
+	FDetailWidgetRow& ShortcutsRow = SwitchboardCategory.AddCustomRow(HealthFilterString, false);
+
+	ShortcutsRow.VisibilityAttr.Bind(TAttribute<EVisibility>::FGetter::CreateLambda([]() -> EVisibility
+		{
+			if (FSwitchboardEditorModule::Get().GetSwitchboardInstallState() ==
+				FSwitchboardEditorModule::ESwitchboardInstallState::ShortcutsMissing)
+			{
+				return EVisibility::Visible;
+			}
+			else
+			{
+				return EVisibility::Collapsed;
+			}
+		}
+	));
+
+	ShortcutsRow.WholeRowWidget
+	[
+		SNew(SHorizontalBox)
+		+ SHorizontalBox::Slot()
+		.Padding(0.f, 4.f, 0.f, 0.f)
+		[
+			SNew(SBorder)
+			.BorderImage(FSwitchboardEditorStyle::Get().GetBrush("Settings.RowBorder.Warning"))
+			.Padding(FMargin(8.0f, 16.0f))
+			[
+				SNew(SHorizontalBox)
+				+ SHorizontalBox::Slot()
+				.AutoWidth()
+				.VAlign(VAlign_Center)
+				.Padding(FMargin(8.0f, 0.0f, 0.0f, 0.0f))
+				[
+					SNew(SImage)
+					.Image(FSwitchboardEditorStyle::Get().GetBrush("Settings.Icons.Warning"))
+				]
+				+ SHorizontalBox::Slot()
+				.FillWidth(1.0f)
+				.VAlign(VAlign_Center)
+				.Padding(16.0f, 0.0f)
+				[
+					SNew(STextBlock)
+					.AutoWrapText(true)
+					.Text(LOCTEXT("AddShortcutsHint", "Add shortcuts to launch Switchboard from the Desktop and Start Menu."))
+				]
+				+ SHorizontalBox::Slot()
+				.AutoWidth()
+				.VAlign(VAlign_Center)
+				[
+					SNew(SButton)
+					.Text(LOCTEXT("AddShortcutsButton", "Add Shortcuts"))
+					.OnClicked_Lambda([this]() {
+						SSwitchboardSetupWizard::OpenWindow(SSwitchboardSetupWizard::EWizardPage::Shortcuts);
+						return FReply::Handled();
+					})
+				]
+			]
+		]
+	];
+
 #if SB_LISTENER_AUTOLAUNCH
 	// This checkbox isn't backed by a property; it directly reflects the state
 	// of a Windows registry entry, and therefore only exists on Windows.
diff --git a/Engine/Plugins/VirtualProduction/Switchboard/Source/SwitchboardEditor/Private/SwitchboardSetupWizard.cpp b/Engine/Plugins/VirtualProduction/Switchboard/Source/SwitchboardEditor/Private/SwitchboardSetupWizard.cpp
--- a/Engine/Plugins/VirtualProduction/Switchboard/Source/SwitchboardEditor/Private/SwitchboardSetupWizard.cpp
+++ b/Engine/Plugins/VirtualProduction/Switchboard/Source/SwitchboardEditor/Private/SwitchboardSetupWizard.cpp
@@ -98,7 +98,7 @@ TSharedRef<SWidget> SSwitchboardSetupWizard::Construct_Page_Intro()
 		.AutoHeight()
 		[
 			SNew(STextBlock)
-			.Text(LOCTEXT("Wizard_Intro_TextBlock", "Switchboard requires the following library/framework dependencies to launch:\n\n\t\u25CF  PySide2\n\n\t\u25CF  python-osc\n\n\t\u25CF  requests\n\n\t\u25CF  six\n\n\t\u25CF  cwRsync\n\n\nUnreal Engine will install the dependencies into the following directory:\n"))
+			.Text(LOCTEXT("Wizard_Intro_TextBlock", "Switchboard requires the following library/framework dependencies to launch:\n\n\t\u25CF  aioquic\n\n\t\u25CF  PySide2\n\n\t\u25CF  python-osc\n\n\t\u25CF  requests\n\n\t\u25CF  six\n\n\t\u25CF  cwRsync\n\n\nUnreal Engine will install the dependencies into the following directory:\n"))
 		]
 		+ SVerticalBox::Slot()
 		.AutoHeight()
