Switchboard: Title bar now shows UE version (from Build.version).

diff --git a/Engine/Plugins/VirtualProduction/Switchboard/Source/Switchboard/switchboard/switchboard_dialog.py b/Engine/Plugins/VirtualProduction/Switchboard/Source/Switchboard/switchboard/switchboard_dialog.py
--- a/Engine/Plugins/VirtualProduction/Switchboard/Source/Switchboard/switchboard/switchboard_dialog.py
+++ b/Engine/Plugins/VirtualProduction/Switchboard/Source/Switchboard/switchboard/switchboard_dialog.py
@@ -9,7 +9,8 @@ import threading
 import time
 import re
 import sys
-from typing import Callable, List, Optional, Set, Union
+import json
+from typing import List, Optional, Set, Union
 
 from pathlib import Path
 
@@ -955,14 +956,41 @@ class SwitchboardDialog(QtCore.QObject):
         self.refresh_window_title()
 
     def refresh_window_title(self):
-        ''' Updates the window title based on the project name '''
+        ''' Updates the window title based on the project name and Unreal Engine version '''
+
+        title = ['Switchboard']
+
+        # The objective is to show the UE version associated with Switchboard,
+        # which is not guaranteed to be the same as the engine_dir settings.
+        # We will rely on the location of this application, and get the associated
+        # Build.version file to get the version. An alternative to this could be to
+        # query it from the UnrealEditor executable (not SBL due to possible mis-redeployment).
+        
+        sbmod_dir = Path(os.path.dirname(sys.modules['switchboard'].__file__))
+
+        for parent in sbmod_dir.parents:
+            if parent.name == "Engine":
+                buildver_path = parent / "Build" / "Build.version"
+                try:
+                    buildver_data = json.load(open(buildver_path, "r"))
+                    buildver = [
+                        f"{buildver_data['MajorVersion']}.{buildver_data['MinorVersion']}.{buildver_data['PatchVersion']}",
+                        str(buildver_data['Changelist']),
+                        str(buildver_data['BranchName'])
+                    ]
+
+                    title.append(f" (Unreal Engine {'-'.join(buildver)})")
+                except Exception as e:
+                    LOGGER.warning(f"Error while parsing {buildver_path}: {e}")
+
+                break
 
         project_name = CONFIG.PROJECT_NAME.get_value()
 
         if project_name:
-            self.window.setWindowTitle(f"Switchboard - {project_name}")
-        else:
-            self.window.setWindowTitle(f"Switchboard")
+            title.append(f' - {project_name}')
+            
+        self.window.setWindowTitle(''.join(title))
 
     def update_current_config_text(self):
         # Can be none when current file is deleted
