Switchboard: Always pick a primary node when launching the cluster.

We've seen a common situation where the user launches a subset of the cluster but forgets to assign a primary node. This results in a failed launch but it is not immediately evident why.

With this change, this situation should never happen again.

diff --git a/Engine/Plugins/VirtualProduction/Switchboard/Source/Switchboard/switchboard/devices/ndisplay/plugin_ndisplay.py b/Engine/Plugins/VirtualProduction/Switchboard/Source/Switchboard/switchboard/devices/ndisplay/plugin_ndisplay.py
--- a/Engine/Plugins/VirtualProduction/Switchboard/Source/Switchboard/switchboard/devices/ndisplay/plugin_ndisplay.py
+++ b/Engine/Plugins/VirtualProduction/Switchboard/Source/Switchboard/switchboard/devices/ndisplay/plugin_ndisplay.py
@@ -1417,12 +1417,41 @@ class DevicenDisplay(DeviceUnreal):
         return os.path.normpath(
             os.path.join(project_dir, 'Content', f'{package_rel_path}.uasset'))
 
+    def get_connected_devices(self):
+        ''' Returns a list with the connected devices/nodes
+        '''
+        nodes = []
+        for device in self.active_unreal_devices:
+            is_device_connected = (device.status == DeviceStatus.CLOSED) or (device.status == DeviceStatus.OPEN)
+            if (device.device_type == "nDisplay") and is_device_connected:
+                nodes.append(device)
+        return nodes
+
     def launch(self, map_name):
+
         if not self.check_settings_valid():
             LOGGER.error(f"{self.name}: Not launching due to invalid settings")
             self.widget._close()
             return
 
+        # Ensure that one of the devices is set as the primary node.
+        nodes = self.get_connected_devices()
+
+        assert len(nodes) > 0  # we wouldn't be launching otherwise
+
+        primary_name = DevicenDisplay.csettings['primary_device_name'].get_value()
+        primary = None
+
+        for node in nodes:
+            if node.name == primary_name:
+                primary = node
+                break
+
+        # if there is no primary, assign one
+        if primary is None:
+            nodes[0].select_as_primary()
+
+        # cache the map to launch
         self.map_name_to_launch = map_name
 
         # Update settings controlled exclusively by the nDisplay config file.
