Switchboard: Sound disabled by default when launching nDisplay

Sound is typically not used, and can add unnecessary overhead.

diff --git a/Engine/Plugins/VirtualProduction/Switchboard/Source/Switchboard/switchboard/devices/ndisplay/plugin_ndisplay.py b/Engine/Plugins/VirtualProduction/Switchboard/Source/Switchboard/switchboard/devices/ndisplay/plugin_ndisplay.py
--- a/Engine/Plugins/VirtualProduction/Switchboard/Source/Switchboard/switchboard/devices/ndisplay/plugin_ndisplay.py
+++ b/Engine/Plugins/VirtualProduction/Switchboard/Source/Switchboard/switchboard/devices/ndisplay/plugin_ndisplay.py
@@ -318,6 +318,11 @@ class DevicenDisplay(DeviceUnreal):
             nice_name="Texture Streaming",
             value=True,
         ),
+        'sound': BoolSetting(
+            attr_name="sound",
+            nice_name="Sound",
+            value=False,
+        ),
         'render_api': OptionSetting(
             attr_name="render_api",
             nice_name="Render API",
@@ -814,6 +819,12 @@ class DevicenDisplay(DeviceUnreal):
                 self.name)
             else "")
 
+        # Sound
+        no_sound = (
+            "-nosound"
+            if not DevicenDisplay.csettings['sound'].get_value(self.name)
+            else "")
+
         # MaxGPUCount (mGPU)
         max_gpu_count = DevicenDisplay.csettings["max_gpu_count"].get_value(
             self.name)
@@ -929,6 +940,7 @@ class DevicenDisplay(DeviceUnreal):
             f'{render_mode}',             # mono/...
             f'{use_all_cores}',           # -useallavailablecores
             f'{no_texture_streaming}',    # -notexturestreaming
+            f'{no_sound}',                # -nosound
             f'-dc_node={self.name}',      # name of this node in the nDisplay cluster
             f'Log={self.log_filename}',   # log file
             f'{ini_engine}',              # Engine ini injections
