Switchboard: Ethernet sync policy doesn't need to be lower case anymore.

Sync policy factory TMap is case insensitive. When uasset is saved, it is using Ethernet (not ethernet), so updating this for consistency.

#rnx

diff --git a/Engine/Plugins/VirtualProduction/Switchboard/Source/Switchboard/switchboard/devices/ndisplay/plugin_ndisplay.py b/Engine/Plugins/VirtualProduction/Switchboard/Source/Switchboard/switchboard/devices/ndisplay/plugin_ndisplay.py
--- a/Engine/Plugins/VirtualProduction/Switchboard/Source/Switchboard/switchboard/devices/ndisplay/plugin_ndisplay.py
+++ b/Engine/Plugins/VirtualProduction/Switchboard/Source/Switchboard/switchboard/devices/ndisplay/plugin_ndisplay.py
@@ -1295,7 +1295,7 @@ class DevicenDisplay(DeviceUnreal):
 
         elif render_sync_policy == 'Ethernet':
 
-            new_rsp['type'] = 'ethernet' # .ndisplay uses lowercase
+            new_rsp['type'] = 'Ethernet'
             new_rsp['parameters'] = {}
 
         elif render_sync_policy == 'None':
