Hash the endpoint id for the server to reduce name collisions.

diff --git a/Engine/Plugins/VirtualProduction/Switchboard/Source/Switchboard/switchboard/switchboard_application.py b/Engine/Plugins/VirtualProduction/Switchboard/Source/Switchboard/switchboard/switchboard_application.py
--- a/Engine/Plugins/VirtualProduction/Switchboard/Source/Switchboard/switchboard/switchboard_application.py
+++ b/Engine/Plugins/VirtualProduction/Switchboard/Source/Switchboard/switchboard/switchboard_application.py
@@ -1,5 +1,7 @@
 # Copyright Epic Games, Inc. All Rights Reserved.
 
+import base64
+import hashlib
 import os
 import pathlib
 import re
@@ -196,7 +198,7 @@ class MultiUserApplication:
                 args.insert(0, f'{self.exe_path()}')
 
             args.extend([
-                f'-CONCERTSERVER="{CONFIG.MUSERVER_SERVER_NAME.get_value()}"',
+                f'-CONCERTSERVER="{self.configured_server_name()}"',
                 f'{CONFIG.MUSERVER_COMMAND_LINE_ARGUMENTS.get_value()}',
                 f'-UDPMESSAGING_SHARE_KNOWN_NODES=1',
                 f'{self.get_mu_server_endpoint_arg()}',
@@ -261,7 +263,20 @@ class MultiUserApplication:
         return endpoint
 
     def configured_server_name(self):
-        return CONFIG.MUSERVER_SERVER_NAME.get_value()
+        endpoint = self.configured_endpoint()
+        if not endpoint:
+            LOGGER.warning(f'No configured endpoint for Multi-user server is defined. This can result in name collisions.')
+            endpoint = '127.0.0.1:9030'
+
+        digest = hashlib.md5(endpoint.encode('utf-8')).digest()
+
+        # Use only the first 12 characters which should be enough to prevent clashing and keep the server name relatively short.
+        #
+        encoded = base64.b64encode(digest).decode('utf-8').strip('==')[:12]
+
+        configured_server_name = CONFIG.MUSERVER_SERVER_NAME.get_value()
+        return f'{configured_server_name}_{encoded}'
+
 
     def clear_process_info(self):
         self._running_endpoint = ""
