Fix multi-user server name on node launch

diff --git a/Engine/Plugins/VirtualProduction/Switchboard/Source/Switchboard/switchboard/devices/ndisplay/plugin_ndisplay.py b/Engine/Plugins/VirtualProduction/Switchboard/Source/Switchboard/switchboard/devices/ndisplay/plugin_ndisplay.py
--- a/Engine/Plugins/VirtualProduction/Switchboard/Source/Switchboard/switchboard/devices/ndisplay/plugin_ndisplay.py
+++ b/Engine/Plugins/VirtualProduction/Switchboard/Source/Switchboard/switchboard/devices/ndisplay/plugin_ndisplay.py
@@ -1004,9 +1004,9 @@ class DevicenDisplay(DeviceUnreal):
             args.extend([
                 '-CONCERTRETRYAUTOCONNECTONERROR',
                 '-CONCERTAUTOCONNECT'])
-
+        mu_server = switchboard_application.get_multi_user_server_instance()
         args.extend([
-            f'-CONCERTSERVER="{CONFIG.MUSERVER_SERVER_NAME.get_value()}"',
+            f'-CONCERTSERVER="{mu_server.configured_server_name()}"',
             f'-CONCERTSESSION="{SETTINGS.MUSERVER_SESSION_NAME}"',
             f'-CONCERTDISPLAYNAME="{self.name}"',
             '-CONCERTISHEADLESS',
diff --git a/Engine/Plugins/VirtualProduction/Switchboard/Source/Switchboard/switchboard/devices/unreal/plugin_unreal.py b/Engine/Plugins/VirtualProduction/Switchboard/Source/Switchboard/switchboard/devices/unreal/plugin_unreal.py
--- a/Engine/Plugins/VirtualProduction/Switchboard/Source/Switchboard/switchboard/devices/unreal/plugin_unreal.py
+++ b/Engine/Plugins/VirtualProduction/Switchboard/Source/Switchboard/switchboard/devices/unreal/plugin_unreal.py
@@ -1950,8 +1950,8 @@ class DeviceUnreal(Device):
             command_line_args += (
                 '-CONCERTRETRYAUTOCONNECTONERROR '
                 '-CONCERTAUTOCONNECT ')
-
-        command_line_args += (f'-CONCERTSERVER="{CONFIG.MUSERVER_SERVER_NAME.get_value()}" '
+        mu_server = switchboard_application.get_multi_user_server_instance()
+        command_line_args += (f'-CONCERTSERVER="{mu_server.configured_server_name()}" '
                               f'-CONCERTSESSION="{SETTINGS.MUSERVER_SESSION_NAME}" '
                               f'-CONCERTDISPLAYNAME="{self.name}"')
 
