Add missing python dependency

diff --git a/Engine/Plugins/VirtualProduction/Switchboard/Source/Switchboard/switchboard/devices/ndisplay/plugin_ndisplay.py b/Engine/Plugins/VirtualProduction/Switchboard/Source/Switchboard/switchboard/devices/ndisplay/plugin_ndisplay.py
--- a/Engine/Plugins/VirtualProduction/Switchboard/Source/Switchboard/switchboard/devices/ndisplay/plugin_ndisplay.py
+++ b/Engine/Plugins/VirtualProduction/Switchboard/Source/Switchboard/switchboard/devices/ndisplay/plugin_ndisplay.py
@@ -12,7 +12,7 @@ from typing import Optional
 from PySide2 import QtCore
 from PySide2 import QtWidgets
 
-from switchboard import message_protocol
+from switchboard import message_protocol, switchboard_application
 from switchboard import switchboard_utils as sb_utils
 from switchboard import switchboard_widgets as sb_widgets
 from switchboard import switchboard_dialog as sb_dialog
