Switchboard: Don't include project path in command line when launching cooked nDisplay builds.

diff --git a/Engine/Plugins/VirtualProduction/Switchboard/Source/Switchboard/switchboard/devices/ndisplay/plugin_ndisplay.py b/Engine/Plugins/VirtualProduction/Switchboard/Source/Switchboard/switchboard/devices/ndisplay/plugin_ndisplay.py
--- a/Engine/Plugins/VirtualProduction/Switchboard/Source/Switchboard/switchboard/devices/ndisplay/plugin_ndisplay.py
+++ b/Engine/Plugins/VirtualProduction/Switchboard/Source/Switchboard/switchboard/devices/ndisplay/plugin_ndisplay.py
@@ -782,8 +782,26 @@ class DevicenDisplay(DeviceUnreal):
 
         return '?'.join([current_map_name, 'Listen', get_common_args()])
 
+    def should_use_project_path_in_command_line(self):
+        ''' Returns true if the project path should be added to the command line.
+        This is normally true when launching the editor, but not when launching a cooked game executable.
+        '''
+ 
+        # The default exe of an Unreal device is the Editor.
+        default_exe = Path(DeviceUnreal.csettings['ue_exe']._original_value)
+
+        # This is the current executable
+        exe = Path(self.generate_unreal_exe_path())
+
+        # We don't use direct comparison to include editor build variants, such as -Debug builds.
+        return exe.stem.lower().startswith(default_exe.stem.lower())
+
     def generate_unreal_command_line(self, map_name=""):
-        uproject = CONFIG.UPROJECT_PATH.get_value(self.name)
+
+        uproject = ''
+
+        if self.should_use_project_path_in_command_line():
+            uproject = CONFIG.UPROJECT_PATH.get_value(self.name)
         
         # normalize path if not empty
         if uproject != "":
