Switchboard plugin: Factor Listener "launch on login" helper methods out into new SwitchboardCommon runtime module for inclusion in SBL target.

#rnx

diff --git a/Engine/Plugins/VirtualProduction/Switchboard/Source/SwitchboardCommon/Internal/SwitchboardListenerAutolaunch.h b/Engine/Plugins/VirtualProduction/Switchboard/Source/SwitchboardCommon/Internal/SwitchboardListenerAutolaunch.h
new file mode 100644
--- /dev/null
+++ b/Engine/Plugins/VirtualProduction/Switchboard/Source/SwitchboardCommon/Internal/SwitchboardListenerAutolaunch.h
@@ -0,0 +1,26 @@
+// Copyright Epic Games, Inc. All Rights Reserved.
+
+#pragma once
+
+#include "Containers/UnrealString.h"
+
+
+#if PLATFORM_WINDOWS
+#	define SWITCHBOARD_LISTENER_AUTOLAUNCH 1
+#else
+#	define SWITCHBOARD_LISTENER_AUTOLAUNCH 0
+#endif
+
+
+struct FLogCategoryBase;
+
+
+namespace UE::SwitchboardListener::Autolaunch
+{
+#if SWITCHBOARD_LISTENER_AUTOLAUNCH
+	SWITCHBOARDCOMMON_API FString GetInvocation(FLogCategoryBase& InLogCategory);
+	SWITCHBOARDCOMMON_API FString GetInvocationExecutable(FLogCategoryBase& InLogCategory);
+	SWITCHBOARDCOMMON_API bool SetInvocation(const FString& InNewInvocation, FLogCategoryBase& InLogCategory);
+	SWITCHBOARDCOMMON_API bool RemoveInvocation(FLogCategoryBase& InLogCategory);
+#endif // #if SWITCHBOARD_LISTENER_AUTOLAUNCH
+}
diff --git a/Engine/Plugins/VirtualProduction/Switchboard/Source/SwitchboardCommon/Private/SwitchboardListenerAutolaunch.cpp b/Engine/Plugins/VirtualProduction/Switchboard/Source/SwitchboardCommon/Private/SwitchboardListenerAutolaunch.cpp
new file mode 100644
--- /dev/null
+++ b/Engine/Plugins/VirtualProduction/Switchboard/Source/SwitchboardCommon/Private/SwitchboardListenerAutolaunch.cpp
@@ -0,0 +1,150 @@
+// Copyright Epic Games, Inc. All Rights Reserved.
+
+#include "SwitchboardListenerAutolaunch.h"
+
+#include "Logging/LogMacros.h"
+#include "Misc/ScopeExit.h"
+
+#if PLATFORM_WINDOWS
+#include "Windows/AllowWindowsPlatformTypes.h"
+THIRD_PARTY_INCLUDES_START
+#include <winreg.h>
+THIRD_PARTY_INCLUDES_END
+#include "Windows/HideWindowsPlatformTypes.h"
+
+static const FString WindowsRunRegKeyPath = "Software\\Microsoft\\Windows\\CurrentVersion\\Run";
+static const FString ListenerAutolaunchRegValueName = "SwitchboardListener";
+#endif
+
+
+namespace UE::SwitchboardListener::Autolaunch
+{
+#if SWITCHBOARD_LISTENER_AUTOLAUNCH
+	HKEY OpenHkcuRunKey(FLogCategoryBase& InLogCategory)
+	{
+		HKEY HkcuRunKey = nullptr;
+		const LSTATUS OpenResult = RegOpenKeyEx(HKEY_CURRENT_USER, *WindowsRunRegKeyPath, 0, KEY_ALL_ACCESS, &HkcuRunKey);
+		if (OpenResult != ERROR_SUCCESS)
+		{
+			UE_LOG_REF(InLogCategory, Log, TEXT("Error opening registry key %s (%08X)"), *WindowsRunRegKeyPath, OpenResult);
+			return nullptr;
+		}
+		return HkcuRunKey;
+	}
+
+
+	FString GetInvocation(FLogCategoryBase& InLogCategory)
+	{
+		HKEY HkcuRunKey = OpenHkcuRunKey(InLogCategory);
+		if (!HkcuRunKey)
+		{
+			return FString();
+		}
+
+		ON_SCOPE_EXIT
+		{
+			RegCloseKey(HkcuRunKey);
+		};
+
+		DWORD ValueType;
+		DWORD ValueSizeBytes = 0;
+		const LSTATUS SizeResult = RegQueryValueEx(HkcuRunKey, *ListenerAutolaunchRegValueName, nullptr, &ValueType, nullptr, &ValueSizeBytes);
+		if (SizeResult == ERROR_FILE_NOT_FOUND)
+		{
+			return FString();
+		}
+		else if (SizeResult != ERROR_SUCCESS)
+		{
+			UE_LOG_REF(InLogCategory, Log, TEXT("Error reading registry value %s:\"%s\" (%08X)"), *WindowsRunRegKeyPath, *ListenerAutolaunchRegValueName, SizeResult);
+			return FString();
+		}
+
+		FString Value;
+		TArray<TCHAR, FString::AllocatorType>& CharArray = Value.GetCharArray();
+		const uint32 ValueLenChars = ValueSizeBytes / sizeof(TCHAR);
+		CharArray.SetNumUninitialized(ValueLenChars);
+
+		const LSTATUS QueryResult = RegQueryValueEx(HkcuRunKey, *ListenerAutolaunchRegValueName, nullptr, &ValueType, reinterpret_cast<LPBYTE>(CharArray.GetData()), &ValueSizeBytes);
+		if (QueryResult != ERROR_SUCCESS)
+		{
+			UE_LOG_REF(InLogCategory, Log, TEXT("Error reading registry value %s:\"%s\" (%08X)"), *WindowsRunRegKeyPath, *ListenerAutolaunchRegValueName, QueryResult);
+			return FString();
+		}
+		else if (ValueType != REG_SZ)
+		{
+			UE_LOG_REF(InLogCategory, Log, TEXT("Registry value %s:\"%s\" has wrong type (%u)"), *WindowsRunRegKeyPath, *ListenerAutolaunchRegValueName, ValueType);
+			return FString();
+		}
+
+		if (CharArray[CharArray.Num() - 1] != TEXT('\0'))
+		{
+			CharArray.Add(TEXT('\0'));
+		}
+
+		return Value;
+	}
+
+
+	FString GetInvocationExecutable(FLogCategoryBase& InLogCategory)
+	{
+		FString AutolaunchCommand = GetInvocation(InLogCategory);
+		if (AutolaunchCommand.IsEmpty())
+		{
+			return FString();
+		}
+
+		TArray<FString> QuoteParts;
+		AutolaunchCommand.ParseIntoArray(QuoteParts, TEXT("\""));
+		return QuoteParts[0];
+	}
+
+
+	bool SetInvocation(const FString& InNewInvocation, FLogCategoryBase& InLogCategory)
+	{
+		HKEY HkcuRunKey = OpenHkcuRunKey(InLogCategory);
+		if (!HkcuRunKey)
+		{
+			return false;
+		}
+
+		ON_SCOPE_EXIT
+		{
+			RegCloseKey(HkcuRunKey);
+		};
+
+		const TArray<TCHAR, FString::AllocatorType>& CharArray = InNewInvocation.GetCharArray();
+		const LSTATUS SetResult = RegSetValueEx(HkcuRunKey, *ListenerAutolaunchRegValueName, 0, REG_SZ, reinterpret_cast<const BYTE*>(CharArray.GetData()), CharArray.Num() * sizeof(TCHAR));
+		if (SetResult != ERROR_SUCCESS)
+		{
+			UE_LOG_REF(InLogCategory, Error, TEXT("Error setting registry value %s:\"%s\" (%08X)"), *WindowsRunRegKeyPath, *ListenerAutolaunchRegValueName, SetResult);
+			return false;
+		}
+
+		return true;
+	}
+
+
+	bool RemoveInvocation(FLogCategoryBase& InLogCategory)
+	{
+		HKEY HkcuRunKey = OpenHkcuRunKey(InLogCategory);
+		if (!HkcuRunKey)
+		{
+			return true;
+		}
+
+		ON_SCOPE_EXIT
+		{
+			RegCloseKey(HkcuRunKey);
+		};
+
+		const LSTATUS DeleteResult = RegDeleteValue(HkcuRunKey, *ListenerAutolaunchRegValueName);
+		if (DeleteResult != ERROR_SUCCESS)
+		{
+			UE_LOG_REF(InLogCategory, Error, TEXT("Error deleting registry value %s:\"%s\" (%08X)"), *WindowsRunRegKeyPath, *ListenerAutolaunchRegValueName, DeleteResult);
+			return false;
+		}
+
+		return true;
+	}
+#endif // #if SWITCHBOARD_LISTENER_AUTOLAUNCH
+} // namespace UE::SwitchboardListener::Autolaunch
diff --git a/Engine/Plugins/VirtualProduction/Switchboard/Source/SwitchboardCommon/SwitchboardCommon.build.cs b/Engine/Plugins/VirtualProduction/Switchboard/Source/SwitchboardCommon/SwitchboardCommon.build.cs
new file mode 100644
--- /dev/null
+++ b/Engine/Plugins/VirtualProduction/Switchboard/Source/SwitchboardCommon/SwitchboardCommon.build.cs
@@ -0,0 +1,14 @@
+// Copyright Epic Games, Inc. All Rights Reserved.
+
+using UnrealBuildTool;
+
+public class SwitchboardCommon : ModuleRules
+{
+	public SwitchboardCommon(ReadOnlyTargetRules Target) : base(Target)
+	{
+		PublicDependencyModuleNames.AddRange(
+			new string[] {
+				"Core",
+			});
+	}
+}
diff --git a/Engine/Plugins/VirtualProduction/Switchboard/Source/SwitchboardEditor/Private/SwitchboardEditorModule.cpp b/Engine/Plugins/VirtualProduction/Switchboard/Source/SwitchboardEditor/Private/SwitchboardEditorModule.cpp
--- a/Engine/Plugins/VirtualProduction/Switchboard/Source/SwitchboardEditor/Private/SwitchboardEditorModule.cpp
+++ b/Engine/Plugins/VirtualProduction/Switchboard/Source/SwitchboardEditor/Private/SwitchboardEditorModule.cpp
@@ -19,6 +19,7 @@
 #include "Misc/ScopeExit.h"
 #include "PropertyEditorModule.h"
 #include "SwitchboardEditorSettings.h"
+#include "SwitchboardListenerAutolaunch.h"
 #include "SwitchboardMenuEntry.h"
 #include "SwitchboardProjectSettings.h"
 #include "SwitchboardSettingsCustomization.h"
@@ -26,22 +27,6 @@
 
 #include <filesystem>
 
-#if PLATFORM_WINDOWS
-#include "Windows/AllowWindowsPlatformTypes.h"
-THIRD_PARTY_INCLUDES_START
-#include <winreg.h>
-THIRD_PARTY_INCLUDES_END
-#include "Windows/HideWindowsPlatformTypes.h"
-
-static const FString WindowsRunRegKeyPath = "Software\\Microsoft\\Windows\\CurrentVersion\\Run";
-static const FString ListenerAutolaunchRegValueName = "SwitchboardListener";
-
-FString GetListenerAutolaunchEntry();
-FString GetListenerAutolaunchEntryExecutable();
-bool SetListenerAutolaunchEntry(const FString& NewCommandLine);
-bool RemoveListenerAutolaunchEntry();
-#endif
-
 
 #define LOCTEXT_NAMESPACE "SwitchboardEditorModule"
 
@@ -140,7 +125,7 @@ void FSwitchboardEditorModule::StartupModule()
 
 	DeferredStartDelegateHandle = FCoreDelegates::OnFEngineLoopInitComplete.AddRaw(this, &FSwitchboardEditorModule::OnEngineInitComplete);
 
-#if SB_LISTENER_AUTOLAUNCH
+#if SWITCHBOARD_LISTENER_AUTOLAUNCH
 	bCachedAutolaunchEnabled = GetListenerAutolaunchEnabled_Internal();
 #endif
 }
@@ -382,7 +367,7 @@ TSharedFuture<FSwitchboardVerifyResult> FSwitchboardEditorModule::GetVerifyResul
 }
 
 
-#if SB_LISTENER_AUTOLAUNCH
+#if SWITCHBOARD_LISTENER_AUTOLAUNCH
 bool FSwitchboardEditorModule::IsListenerAutolaunchEnabled(bool bForceRefreshCache /* = false */)
 {
 	if (bForceRefreshCache)
@@ -396,7 +381,7 @@ bool FSwitchboardEditorModule::IsListenerAutolaunchEnabled(bool bForceRefreshCac
 
 bool FSwitchboardEditorModule::GetListenerAutolaunchEnabled_Internal() const
 {
-	const FString& ExistingCmd = GetListenerAutolaunchEntry().TrimEnd();
+	const FString& ExistingCmd = UE::SwitchboardListener::Autolaunch::GetInvocation(LogSwitchboardPlugin).TrimEnd();
 	const FString& ConfigCmd = GetDefault<USwitchboardEditorSettings>()->GetListenerInvocation().TrimEnd();
 	return ExistingCmd == ConfigCmd;
 }
@@ -409,17 +394,17 @@ bool FSwitchboardEditorModule::SetListenerAutolaunchEnabled(bool bEnabled)
 	if (bEnabled)
 	{
 		const FString CommandLine = GetDefault<USwitchboardEditorSettings>()->GetListenerInvocation();
-		bSucceeded = SetListenerAutolaunchEntry(CommandLine);
+		bSucceeded = UE::SwitchboardListener::Autolaunch::SetInvocation(CommandLine, LogSwitchboardPlugin);
 	}
 	else
 	{
-		bSucceeded = RemoveListenerAutolaunchEntry();
+		bSucceeded = UE::SwitchboardListener::Autolaunch::RemoveInvocation(LogSwitchboardPlugin);
 	}
 
 	bCachedAutolaunchEnabled = GetListenerAutolaunchEnabled_Internal();
 	return bSucceeded;
 }
-#endif // #if SB_LISTENER_AUTOLAUNCH
+#endif // #if SBL_AUTOLAUNCH
 
 
 #if SWITCHBOARD_SHORTCUTS
@@ -477,10 +462,10 @@ void FSwitchboardEditorModule::OnEngineInitComplete()
 
 bool FSwitchboardEditorModule::OnEditorSettingsModified()
 {
-#if SB_LISTENER_AUTOLAUNCH
+#if SWITCHBOARD_LISTENER_AUTOLAUNCH
 	// If the existing entry's listener executable path matches our current engine / path,
 	// then we ensure the command line arguments also stay in sync with our settings.
-	const FString AutolaunchEntryExecutable = GetListenerAutolaunchEntryExecutable();
+	const FString AutolaunchEntryExecutable = UE::SwitchboardListener::Autolaunch::GetInvocationExecutable(LogSwitchboardPlugin);
 	if (!AutolaunchEntryExecutable.IsEmpty())
 	{
 		const FString ConfigExecutable = GetDefault<USwitchboardEditorSettings>()->GetListenerPlatformPath();
@@ -495,7 +480,7 @@ bool FSwitchboardEditorModule::OnEditorSettingsModified()
 				*AutolaunchEntryExecutable, *ConfigExecutable);
 		}
 	}
-#endif // #if SB_LISTENER_AUTOLAUNCH
+#endif // #if SWITCHBOARD_LISTENER_AUTOLAUNCH
 
 	return true;
 }
@@ -516,136 +501,6 @@ void FSwitchboardEditorModule::RunDefaultOSCListener()
 }
 
 
-#if PLATFORM_WINDOWS
-HKEY OpenHkcuRunKey()
-{
-	HKEY HkcuRunKey = nullptr;
-	const LSTATUS OpenResult = RegOpenKeyEx(HKEY_CURRENT_USER, *WindowsRunRegKeyPath, 0, KEY_ALL_ACCESS, &HkcuRunKey);
-	if (OpenResult != ERROR_SUCCESS)
-	{
-		UE_LOG(LogSwitchboardPlugin, Log, TEXT("Error opening registry key %s (%08X)"), *WindowsRunRegKeyPath, OpenResult);
-		return nullptr;
-	}
-	return HkcuRunKey;
-}
-
-
-FString GetListenerAutolaunchEntry()
-{
-	HKEY HkcuRunKey = OpenHkcuRunKey();
-	if (!HkcuRunKey)
-	{
-		return FString();
-	}
-
-	ON_SCOPE_EXIT
-	{
-		RegCloseKey(HkcuRunKey);
-	};
-
-	DWORD ValueType;
-	DWORD ValueSizeBytes = 0;
-	const LSTATUS SizeResult = RegQueryValueEx(HkcuRunKey, *ListenerAutolaunchRegValueName, nullptr, &ValueType, nullptr, &ValueSizeBytes);
-	if (SizeResult == ERROR_FILE_NOT_FOUND)
-	{
-		return FString();
-	}
-	else if (SizeResult != ERROR_SUCCESS)
-	{
-		UE_LOG(LogSwitchboardPlugin, Log, TEXT("Error reading registry value %s:\"%s\" (%08X)"), *WindowsRunRegKeyPath, *ListenerAutolaunchRegValueName, SizeResult);
-		return FString();
-	}
-
-	FString Value;
-	TArray<TCHAR, FString::AllocatorType>& CharArray = Value.GetCharArray();
-	const uint32 ValueLenChars = ValueSizeBytes / sizeof(TCHAR);
-	CharArray.SetNumUninitialized(ValueLenChars);
-
-	const LSTATUS QueryResult = RegQueryValueEx(HkcuRunKey, *ListenerAutolaunchRegValueName, nullptr, &ValueType, reinterpret_cast<LPBYTE>(CharArray.GetData()), &ValueSizeBytes);
-	if (QueryResult != ERROR_SUCCESS)
-	{
-		UE_LOG(LogSwitchboardPlugin, Log, TEXT("Error reading registry value %s:\"%s\" (%08X)"), *WindowsRunRegKeyPath, *ListenerAutolaunchRegValueName, QueryResult);
-		return FString();
-	}
-	else if (ValueType != REG_SZ)
-	{
-		UE_LOG(LogSwitchboardPlugin, Log, TEXT("Registry value %s:\"%s\" has wrong type (%u)"), *WindowsRunRegKeyPath, *ListenerAutolaunchRegValueName, ValueType);
-		return FString();
-	}
-
-	if (CharArray[CharArray.Num() - 1] != TEXT('\0'))
-	{
-		CharArray.Add(TEXT('\0'));
-	}
-
-	return Value;
-}
-
-
-FString GetListenerAutolaunchEntryExecutable()
-{
-	FString AutolaunchCommand = GetListenerAutolaunchEntry();
-	if (AutolaunchCommand.IsEmpty())
-	{
-		return FString();
-	}
-
-	TArray<FString> QuoteParts;
-	AutolaunchCommand.ParseIntoArray(QuoteParts, TEXT("\""));
-	return QuoteParts[0];
-}
-
-
-bool SetListenerAutolaunchEntry(const FString& NewCommandLine)
-{
-	HKEY HkcuRunKey = OpenHkcuRunKey();
-	if (!HkcuRunKey)
-	{
-		return false;
-	}
-
-	ON_SCOPE_EXIT
-	{
-		RegCloseKey(HkcuRunKey);
-	};
-
-	const TArray<TCHAR, FString::AllocatorType>& CharArray = NewCommandLine.GetCharArray();
-	const LSTATUS SetResult = RegSetValueEx(HkcuRunKey, *ListenerAutolaunchRegValueName, 0, REG_SZ, reinterpret_cast<const BYTE*>(CharArray.GetData()), CharArray.Num() * sizeof(TCHAR));
-	if (SetResult != ERROR_SUCCESS)
-	{
-		UE_LOG(LogSwitchboardPlugin, Error, TEXT("Error setting registry value %s:\"%s\" (%08X)"), *WindowsRunRegKeyPath, *ListenerAutolaunchRegValueName, SetResult);
-		return false;
-	}
-
-	return true;
-}
-
-
-bool RemoveListenerAutolaunchEntry()
-{
-	HKEY HkcuRunKey = OpenHkcuRunKey();
-	if (!HkcuRunKey)
-	{
-		return true;
-	}
-
-	ON_SCOPE_EXIT
-	{
-		RegCloseKey(HkcuRunKey);
-	};
-
-	const LSTATUS DeleteResult = RegDeleteValue(HkcuRunKey, *ListenerAutolaunchRegValueName);
-	if (DeleteResult != ERROR_SUCCESS)
-	{
-		UE_LOG(LogSwitchboardPlugin, Error, TEXT("Error deleting registry value %s:\"%s\" (%08X)"), *WindowsRunRegKeyPath, *ListenerAutolaunchRegValueName, DeleteResult);
-		return false;
-	}
-
-	return true;
-}
-#endif // #if PLATFORM_WINDOWS
-
-
 IMPLEMENT_MODULE(FSwitchboardEditorModule, SwitchboardEditor);
 
 #undef LOCTEXT_NAMESPACE
diff --git a/Engine/Plugins/VirtualProduction/Switchboard/Source/SwitchboardEditor/Private/SwitchboardEditorModule.h b/Engine/Plugins/VirtualProduction/Switchboard/Source/SwitchboardEditor/Private/SwitchboardEditorModule.h
--- a/Engine/Plugins/VirtualProduction/Switchboard/Source/SwitchboardEditor/Private/SwitchboardEditorModule.h
+++ b/Engine/Plugins/VirtualProduction/Switchboard/Source/SwitchboardEditor/Private/SwitchboardEditorModule.h
@@ -8,6 +8,7 @@
 #include "Misc/Paths.h"
 #include "Modules/ModuleInterface.h"
 #include "Modules/ModuleManager.h"
+#include "SwitchboardListenerAutolaunch.h"
 #include "SwitchboardShortcuts.h"
 #include "SwitchboardScriptInterop.h"
 #include "SwitchboardTypes.h"
@@ -83,7 +84,7 @@ public:
 
 	TSharedFuture<FSwitchboardVerifyResult> GetVerifyResult(bool bForceRefresh = false);
 
-#if SB_LISTENER_AUTOLAUNCH
+#if SWITCHBOARD_LISTENER_AUTOLAUNCH
 	/**
 	 * Returns whether (this engine's) SwitchboardListener is configured to run automatically.
 	 * Defaults to returning a cached value to avoid hitting the registry.
@@ -127,7 +128,7 @@ private:
 	FString VerifyPath;
 	TSharedFuture<FSwitchboardVerifyResult> VerifyResult;
 
-#if SB_LISTENER_AUTOLAUNCH
+#if SWITCHBOARD_LISTENER_AUTOLAUNCH
 	bool GetListenerAutolaunchEnabled_Internal() const;
 
 	bool bCachedAutolaunchEnabled;
diff --git a/Engine/Plugins/VirtualProduction/Switchboard/Source/SwitchboardEditor/Private/SwitchboardMenuEntry.cpp b/Engine/Plugins/VirtualProduction/Switchboard/Source/SwitchboardEditor/Private/SwitchboardMenuEntry.cpp
--- a/Engine/Plugins/VirtualProduction/Switchboard/Source/SwitchboardEditor/Private/SwitchboardMenuEntry.cpp
+++ b/Engine/Plugins/VirtualProduction/Switchboard/Source/SwitchboardEditor/Private/SwitchboardMenuEntry.cpp
@@ -128,7 +128,7 @@ struct FSwitchboardMenuEntryImpl
 		{
 			MenuBuilder.AddMenuEntry(FSwitchboardUICommands::Get().LaunchSwitchboardListener);
 
-#if SB_LISTENER_AUTOLAUNCH
+#if SWITCHBOARD_LISTENER_AUTOLAUNCH
 			MenuBuilder.AddMenuSeparator();
 
 			MenuBuilder.AddMenuEntry(
@@ -146,7 +146,7 @@ struct FSwitchboardMenuEntryImpl
 				NAME_None,
 				EUserInterfaceActionType::ToggleButton
 			);
-#endif // #if SB_LISTENER_AUTOLAUNCH
+#endif // #if SWITCHBOARD_LISTENER_AUTOLAUNCH
 		}
 		MenuBuilder.EndSection();
 
@@ -204,7 +204,7 @@ struct FSwitchboardMenuEntryImpl
 		}
 	}
 
-#if SB_LISTENER_AUTOLAUNCH
+#if SWITCHBOARD_LISTENER_AUTOLAUNCH
 	void ToggleAutolaunch()
 	{
 		if (FSwitchboardEditorModule::Get().IsListenerAutolaunchEnabled())
@@ -223,7 +223,7 @@ struct FSwitchboardMenuEntryImpl
 			FSwitchboardEditorModule::Get().SetListenerAutolaunchEnabled(true);
 		}
 	}
-#endif // #if SB_LISTENER_AUTOLAUNCH
+#endif // #if SWITCHBOARD_LISTENER_AUTOLAUNCH
 
 	void OnLaunchSwitchboard()
 	{
diff --git a/Engine/Plugins/VirtualProduction/Switchboard/Source/SwitchboardEditor/Private/SwitchboardSettingsCustomization.cpp b/Engine/Plugins/VirtualProduction/Switchboard/Source/SwitchboardEditor/Private/SwitchboardSettingsCustomization.cpp
--- a/Engine/Plugins/VirtualProduction/Switchboard/Source/SwitchboardEditor/Private/SwitchboardSettingsCustomization.cpp
+++ b/Engine/Plugins/VirtualProduction/Switchboard/Source/SwitchboardEditor/Private/SwitchboardSettingsCustomization.cpp
@@ -200,7 +200,7 @@ void FSwitchboardEditorSettingsCustomization::CustomizeDetails(IDetailLayoutBuil
 		]
 	];
 
-#if SB_LISTENER_AUTOLAUNCH
+#if SWITCHBOARD_LISTENER_AUTOLAUNCH
 	// This checkbox isn't backed by a property; it directly reflects the state
 	// of a Windows registry entry, and therefore only exists on Windows.
 	const FText& ListenerAutolaunchLabel = LOCTEXT("ListenerAutolaunchLabel", "Launch Switchboard Listener on Login");
@@ -232,7 +232,7 @@ void FSwitchboardEditorSettingsCustomization::CustomizeDetails(IDetailLayoutBuil
 			FSwitchboardEditorModule::Get().SetListenerAutolaunchEnabled(bShouldAutolaunch);
 		})
 	];
-#endif // #if SB_LISTENER_AUTOLAUNCH
+#endif // #if SWITCHBOARD_LISTENER_AUTOLAUNCH
 }
 
 
diff --git a/Engine/Plugins/VirtualProduction/Switchboard/Source/SwitchboardEditor/Private/SwitchboardSetupWizard.cpp b/Engine/Plugins/VirtualProduction/Switchboard/Source/SwitchboardEditor/Private/SwitchboardSetupWizard.cpp
--- a/Engine/Plugins/VirtualProduction/Switchboard/Source/SwitchboardEditor/Private/SwitchboardSetupWizard.cpp
+++ b/Engine/Plugins/VirtualProduction/Switchboard/Source/SwitchboardEditor/Private/SwitchboardSetupWizard.cpp
@@ -367,7 +367,7 @@ TSharedRef<SWidget> SSwitchboardSetupWizard::Construct_Page_Autolaunch()
 			.Padding(24.0f, 0.0f, 12.0f, 0.0f)
 			[
 				SAssignNew(AutolaunchCheckbox, SCheckBox)
-#if SB_LISTENER_AUTOLAUNCH
+#if SWITCHBOARD_LISTENER_AUTOLAUNCH
 				.IsChecked(FSwitchboardEditorModule::Get().IsListenerAutolaunchEnabled() ? ECheckBoxState::Checked : ECheckBoxState::Unchecked)
 #else
 				.IsEnabled(false)
@@ -377,7 +377,7 @@ TSharedRef<SWidget> SSwitchboardSetupWizard::Construct_Page_Autolaunch()
 			.AutoWidth()
 			[
 				SNew(STextBlock)
-#if SB_LISTENER_AUTOLAUNCH
+#if SWITCHBOARD_LISTENER_AUTOLAUNCH
 				.Text(LOCTEXT("Wizard_Autolaunch_CheckboxLabel", "Launch Switchboard Listener on Login"))
 #else
 				.IsEnabled(false)
@@ -401,7 +401,7 @@ TSharedRef<SWidget> SSwitchboardSetupWizard::Construct_Page_Autolaunch()
 				SNew(SPrimaryButton)
 				.Text(LOCTEXT("Wizard_Autolaunch_ContinueButton", "Continue"))
 				.OnClicked_Lambda([this]() {
-#if SB_LISTENER_AUTOLAUNCH
+#if SWITCHBOARD_LISTENER_AUTOLAUNCH
 					FSwitchboardEditorModule::Get().SetListenerAutolaunchEnabled(AutolaunchCheckbox->IsChecked());
 #endif
 					SwitchToPage(EWizardPage::Complete);
diff --git a/Engine/Plugins/VirtualProduction/Switchboard/Source/SwitchboardEditor/SwitchboardEditor.build.cs b/Engine/Plugins/VirtualProduction/Switchboard/Source/SwitchboardEditor/SwitchboardEditor.build.cs
--- a/Engine/Plugins/VirtualProduction/Switchboard/Source/SwitchboardEditor/SwitchboardEditor.build.cs
+++ b/Engine/Plugins/VirtualProduction/Switchboard/Source/SwitchboardEditor/SwitchboardEditor.build.cs
@@ -24,6 +24,7 @@ public class SwitchboardEditor : ModuleRules
 				"Settings",
 				"Slate",
 				"SlateCore",
+				"SwitchboardCommon",
 				"ToolMenus",
 				"ToolWidgets",
 				"UnrealEd",
@@ -31,12 +32,10 @@ public class SwitchboardEditor : ModuleRules
 
 		if (Target.Platform == UnrealTargetPlatform.Win64)
 		{
-			PrivateDefinitions.Add("SB_LISTENER_AUTOLAUNCH=1");
 			PrivateDefinitions.Add("SWITCHBOARD_SHORTCUTS=1");
 		}
 		else
 		{
-			PrivateDefinitions.Add("SB_LISTENER_AUTOLAUNCH=0");
 			PrivateDefinitions.Add("SWITCHBOARD_SHORTCUTS=0");
 		}
 	}
diff --git a/Engine/Plugins/VirtualProduction/Switchboard/Switchboard.uplugin b/Engine/Plugins/VirtualProduction/Switchboard/Switchboard.uplugin
--- a/Engine/Plugins/VirtualProduction/Switchboard/Switchboard.uplugin
+++ b/Engine/Plugins/VirtualProduction/Switchboard/Switchboard.uplugin
@@ -15,6 +15,15 @@
 	"IsBetaVersion": true,
 	"Installed": false,
 	"Modules": [
+		{
+			"Name": "SwitchboardCommon",
+			"Type": "Runtime",
+			"LoadingPhase": "Default",
+			"PlatformAllowList": [
+				"Win64",
+				"Linux"
+			]
+		},
 		{
 			"Name": "SwitchboardEditor",
 			"Type": "Editor",
