SwitchboardCommon: Missed marking a new file for add. Followup to CL 30664927.

#rnx

diff --git a/Engine/Plugins/VirtualProduction/Switchboard/Source/SwitchboardCommon/Private/SwitchboardCommonModule.cpp b/Engine/Plugins/VirtualProduction/Switchboard/Source/SwitchboardCommon/Private/SwitchboardCommonModule.cpp
new file mode 100644
--- /dev/null
+++ b/Engine/Plugins/VirtualProduction/Switchboard/Source/SwitchboardCommon/Private/SwitchboardCommonModule.cpp
@@ -0,0 +1,6 @@
+// Copyright Epic Games, Inc. All Rights Reserved.
+
+#include "Modules/ModuleManager.h"
+
+
+IMPLEMENT_MODULE(FDefaultModuleImpl, SwitchboardCommon);
