Switchboard: Added a cache of the list of maps in a project.

Certain projects that have many folders and files (e.g. City, which has over 100k of each) can take a relatively long time to open because it re-finds the .umaps every time that SB is open, unless the user has taken the time to set a more narrow set of folders to search.

SB will not be aware if new maps have been added to the project, but the user can refresh the map list if they are not finding the map that they are looking for.

A future improvement could be to add a progress window so that the user knows that the folder structure is being traversed and therefore explain why the application is otherwise looking hung up when refreshing the map list.

Also fixed a null object reference exception, and cleaned up some of the indentation in the edited files.

diff --git a/Engine/Plugins/VirtualProduction/Switchboard/Source/Switchboard/switchboard/config.py b/Engine/Plugins/VirtualProduction/Switchboard/Source/Switchboard/switchboard/config.py
--- a/Engine/Plugins/VirtualProduction/Switchboard/Source/Switchboard/switchboard/config.py
+++ b/Engine/Plugins/VirtualProduction/Switchboard/Source/Switchboard/switchboard/config.py
@@ -2091,6 +2091,7 @@ class Config(object):
 
         # MISC SETTINGS
         self.CURRENT_LEVEL = data.get('current_level', DEFAULT_MAP_TEXT)
+        self.LEVELS = data.get('levels', None)
 
         # Devices
         self._device_data_from_config = {}
@@ -2158,6 +2159,7 @@ class Config(object):
         self.init_muserver()
 
         self.CURRENT_LEVEL = DEFAULT_MAP_TEXT
+        self.LEVELS = None
 
         self._device_data_from_config = {}
         self._plugin_data_from_config = {}
@@ -2575,8 +2577,9 @@ class Config(object):
 
         self.save_muserver(data)
 
-        # Current Level
+        # Levels
         data["current_level"] = self.CURRENT_LEVEL
+        data["levels"] = self.LEVELS
 
         # Devices
         data["devices"] = {}
@@ -2745,23 +2748,49 @@ class Config(object):
                 (unreal_content_plugin,
                  unreal_content_plugin.plugin_content_path))
 
+        # Pre-filtering by .umap extension is faster than fnmatch on all files
+        # Currently there are no use cases where a different extension is of interest.
+        def get_umap_files(dirpath) -> list:
+            ''' Returns a list with all *.map files in the recursively searched dirpath
+
+            Parameters
+            ----------
+            dirpath: str
+                Directory to be recursively searched for .umap files
+
+            Returns
+            -------
+                list
+                    List of found map entries.
+            '''
+            umap_files = []
+            for entry in os.scandir(dirpath):
+                if entry.is_file() and entry.name.endswith(".umap"):
+                    umap_files.append(entry)
+                elif entry.is_dir():
+                    umap_files.extend(get_umap_files(entry.path))  # recursive
+            return umap_files
+
         maps = []
+        maps_filter = self.MAPS_FILTER.get_value()
+
         for (unreal_content_plugin, maps_path) in search_paths:
-            for dirpath, b, file_names in os.walk(maps_path):
-                for file_name in file_names:
-                    if not fnmatch.fnmatch(
-                            file_name, self.MAPS_FILTER.get_value()):
-                        continue
 
-                    map_name, _ = os.path.splitext(file_name)
-                    file_path_to_map = os.path.join(dirpath, map_name)
+            umaps = get_umap_files(maps_path)
+
+            for umap in umaps:
+                if not fnmatch.fnmatch(umap.name, maps_filter):
+                    continue
+
+                map_name, _ = os.path.splitext(umap.name)
+                file_path_to_map = os.path.join(umap.path, map_name)
 
-                    content_path_to_map = self.resolve_content_path(
-                        file_path_to_map,
-                        unreal_content_plugin=unreal_content_plugin)
+                content_path_to_map = self.resolve_content_path(
+                    file_path_to_map,
+                    unreal_content_plugin=unreal_content_plugin)
 
-                    if content_path_to_map not in maps:
-                        maps.append(content_path_to_map)
+                if content_path_to_map not in maps:
+                    maps.append(content_path_to_map)
 
         maps.sort()
         return maps
diff --git a/Engine/Plugins/VirtualProduction/Switchboard/Source/Switchboard/switchboard/devices/ndisplay/plugin_ndisplay.py b/Engine/Plugins/VirtualProduction/Switchboard/Source/Switchboard/switchboard/devices/ndisplay/plugin_ndisplay.py
--- a/Engine/Plugins/VirtualProduction/Switchboard/Source/Switchboard/switchboard/devices/ndisplay/plugin_ndisplay.py
+++ b/Engine/Plugins/VirtualProduction/Switchboard/Source/Switchboard/switchboard/devices/ndisplay/plugin_ndisplay.py
@@ -230,11 +230,12 @@ class DeviceWidgetnDisplay(DeviceWidgetUnreal):
         super().on_name_changed(text)
 
         if self.name_line_edit.is_valid:
-            point = self.name_line_edit.parent().mapToGlobal(
-                self.name_line_edit.geometry().topRight())
-            self.help_tool_tip.showText(
-                point,
-                "WARNING: Device names must match the nDisplay configuration")
+            parent = self.name_line_edit.parent()
+            if parent:
+                point = parent.mapToGlobal(self.name_line_edit.geometry().topRight())
+                self.help_tool_tip.showText(
+                    point,
+                    "WARNING: Device names must match the nDisplay configuration")
 
     def _add_control_buttons(self):
         self._autojoin_visible = False
diff --git a/Engine/Plugins/VirtualProduction/Switchboard/Source/Switchboard/switchboard/switchboard_dialog.py b/Engine/Plugins/VirtualProduction/Switchboard/Source/Switchboard/switchboard/switchboard_dialog.py
--- a/Engine/Plugins/VirtualProduction/Switchboard/Source/Switchboard/switchboard/switchboard_dialog.py
+++ b/Engine/Plugins/VirtualProduction/Switchboard/Source/Switchboard/switchboard/switchboard_dialog.py
@@ -48,6 +48,7 @@ ENGINE_PATH = "../../../../.."
 RELATIVE_PATH = os.path.dirname(__file__)
 EMPTY_SYNC_ENTRY = "-- None --"
 
+
 class TraceSettings(QtWidgets.QDialog):
     """
     Custom class to prompt user for the trace arguments to use for Unreal Insights trace collection.
@@ -79,13 +80,14 @@ class TraceSettings(QtWidgets.QDialog):
         """
         return self.arguments_field.text()
 
+
 class DeviceAdditionalSettingsUI(QtCore.QObject):
     signal_device_widget_tracing = QtCore.Signal(object)
 
     enable_insight_trace = False
     insight_trace_args = "log,cpu,gpu,frame,bookmark,concert,messaging"
 
-    def __init__(self, name, parent = None):
+    def __init__(self, name, parent=None):
         super().__init__(parent)
         self.name = name
 
@@ -337,16 +339,16 @@ class SwitchboardDialog(QtCore.QObject):
         # add level picker combo box and refresh button
         self.level_combo_box = sb_widgets.SearchableComboBox(self.window)
         self.level_combo_box.setSizePolicy(QtWidgets.QSizePolicy.Expanding, QtWidgets.QSizePolicy.Fixed)
-        
+
         self.refresh_levels_button = sb_widgets.ControlQPushButton()
         self.refresh_levels_button.setMaximumSize(22, 22)
         self.refresh_levels_button.setIcon(QtGui.QIcon("icon_refresh.png"))
         self.refresh_levels_button.setProperty("frameless", True)
         self.refresh_levels_button.setToolTip("Refresh level list")
-        
+
         self.window.horizontalLayout_7.addWidget(self.level_combo_box)
         self.window.horizontalLayout_7.addWidget(self.refresh_levels_button)
-        
+
         self.shoot = 'Default'
         self.sequence = SETTINGS.CURRENT_SEQUENCE
         self.slate = SETTINGS.CURRENT_SLATE
@@ -442,7 +444,7 @@ class SwitchboardDialog(QtCore.QObject):
         # Connect to the session increment number.
         self.window.multiuser_session_inc_button.clicked.connect(self.on_multiuser_session_inc)
 
-         # Stylesheet-related: Object names used for selectors, no focus forcing
+        # Stylesheet-related: Object names used for selectors, no focus forcing
         def configure_ctrl_btn(btn: sb_widgets.ControlQPushButton, name: str):
             btn.setObjectName(name)
             btn.hover_focus = False
@@ -484,7 +486,7 @@ class SwitchboardDialog(QtCore.QObject):
 
         if CONFIG.file_path:
             self.toggle_p4_controls(CONFIG.P4_ENABLED.get_value())
-            self.refresh_levels()
+            self.refresh_levels(CONFIG.LEVELS)
         else:
             self.menu_new_config()
 
@@ -494,7 +496,7 @@ class SwitchboardDialog(QtCore.QObject):
 
         # Run the transport queue
         #self.transport_queue_resume()
-        
+
         self.update_current_config_text()
         self.update_current_address_text()
         SETTINGS.ADDRESS.signal_setting_changed.connect(
@@ -506,14 +508,14 @@ class SwitchboardDialog(QtCore.QObject):
 
         self.script_manager.on_postinit(self)
         self.have_warned_about_muserver = False
-        
+
     def _try_change_address(self):
         new_value = self.window.current_address_value.text()
         old_address = SETTINGS.ADDRESS.get_value()
-        
+
         if not SETTINGS.ADDRESS.possible_values.__contains__(new_value):
             SETTINGS.ADDRESS.possible_values.append(new_value)
-            
+
         SETTINGS.ADDRESS.update_value(new_value)
         SETTINGS.save()
         # Check if osc connection binds before commiting changes
@@ -581,21 +583,21 @@ class SwitchboardDialog(QtCore.QObject):
             self._on_mu_server_status_changed)
 
     def update_locallistener_menuitem(self):
-        ''' 
+        '''
         Enables/disables the local listener launch menu item depending on whether 
         it is already running or not.
         '''
         self.locallistener_launcher_menuitem.setEnabled(not self.listener_launcher.is_running())
 
     def update_localsblhelper_menuitem(self):
-        ''' 
+        '''
         Enables/disables the local listener helper launch menu item depending on whether 
         it is already running or not.
         '''
         self.localsblhelper_launcher_menuitem.setEnabled(not self.sblhelper_launcher.is_running())
 
     def update_insights_menuitem(self):
-        ''' 
+        '''
         Enables/disables the UnrealInsights launch menu item depending on whether 
         it is already running or not.
         '''
@@ -711,7 +713,7 @@ class SwitchboardDialog(QtCore.QObject):
 
         self.tools_menu = self.window.menu_bar.addMenu("&Tools")
 
-    def register_tools_menu_action(self, actionname:str, menunames:List[str] = []) -> QWidgetAction:
+    def register_tools_menu_action(self, actionname: str, menunames: List[str] = []) -> QWidgetAction:
         ''' Registers a QWidgetAction with the tools menu
 
         Args:
@@ -724,7 +726,7 @@ class SwitchboardDialog(QtCore.QObject):
 
         # Find find submenu, and create the ones that don't exist along the way
 
-        current_menu:QMenu = self.tools_menu
+        current_menu: QMenu = self.tools_menu
 
         # iterate over menunames given
         for menuname in menunames:
@@ -740,7 +742,7 @@ class SwitchboardDialog(QtCore.QObject):
                     current_menu = submenu
                     create_menu = False
                     break
-            
+
             # if we didn't find the submenu, create it
             if create_menu:
                 newmenu = QMenu(parent=current_menu)
@@ -752,7 +754,7 @@ class SwitchboardDialog(QtCore.QObject):
         action = QWidgetAction(current_menu)
         action.setText(actionname)
         current_menu.addAction(action)
-        
+
         return action
 
     def on_device_autojoin_changed(self):
@@ -805,7 +807,7 @@ class SwitchboardDialog(QtCore.QObject):
                 # this offers a simple way to update nDisplay should the config file have been changed.
                 self.device_manager.remove_device(device)
 
-            self.device_manager.add_devices({device_type : dialog.devices_to_add()})
+            self.device_manager.add_devices({device_type: dialog.devices_to_add()})
             CONFIG.save()
 
     def on_device_removed(self, device_hash, device_type, device_name, update_config):
@@ -907,7 +909,8 @@ class SwitchboardDialog(QtCore.QObject):
         ''' When the user wants to open a config located outside of the designated area
         This is useful when the user keeps the configs somewhere else (e.g. under the project)
         '''
-        config_path, _ = QtWidgets.QFileDialog.getOpenFileName(self.window,
+        config_path, _ = QtWidgets.QFileDialog.getOpenFileName(
+            self.window,
             'Select config file', str(config.ROOT_CONFIGS_PATH),
             f'Config files (*{config.CONFIG_SUFFIX})'
         )
@@ -949,7 +952,7 @@ class SwitchboardDialog(QtCore.QObject):
 
         self.p4_refresh_project_cl()
         self.p4_refresh_engine_cl()
-        self.refresh_levels()
+        self.refresh_levels(levels=CONFIG.LEVELS)
         self.update_current_config_text()
         self.refresh_muserver_autojoin()
         self.refresh_trace_settings()
@@ -965,7 +968,7 @@ class SwitchboardDialog(QtCore.QObject):
         # We will rely on the location of this application, and get the associated
         # Build.version file to get the version. An alternative to this could be to
         # query it from the UnrealEditor executable (not SBL due to possible mis-redeployment).
-        
+
         sbmod_dir = Path(os.path.dirname(sys.modules['switchboard'].__file__))
 
         for parent in sbmod_dir.parents:
@@ -989,7 +992,7 @@ class SwitchboardDialog(QtCore.QObject):
 
         if project_name:
             title.append(f' - {project_name}')
-            
+
         self.window.setWindowTitle(''.join(title))
 
     def update_current_config_text(self):
@@ -999,7 +1002,7 @@ class SwitchboardDialog(QtCore.QObject):
             self.window.current_config_file_value.setText(file_name)
         else:
             self.window.current_config_file_value.setText("No config loaded")
-    
+
     def update_current_address_text(self):
         self.window.current_address_value.setText(SETTINGS.ADDRESS.get_value())
 
@@ -1049,10 +1052,10 @@ class SwitchboardDialog(QtCore.QObject):
 
         # get the destination path for the copy
 
-        new_config_path, _ = QtWidgets.QFileDialog.getSaveFileName(self.window,
+        new_config_path, _ = QtWidgets.QFileDialog.getSaveFileName(
+            self.window,
             'Select config file', str(config.ROOT_CONFIGS_PATH),
-            f'Config files (*{config.CONFIG_SUFFIX})',
-            #options=QtWidgets.QFileDialog.DontUseNativeDialog
+            f'Config files (*{config.CONFIG_SUFFIX})'
         )
 
         # Do nothing if the user didn't choose a destination config path
@@ -1111,7 +1114,9 @@ class SwitchboardDialog(QtCore.QObject):
         """
         # Show the confirmation dialog using a relative path to the config.
         rel_config_path = config.get_relative_config_path(SETTINGS.CONFIG)
-        reply = QtWidgets.QMessageBox.question(self.window, 'Delete Config',
+        reply = QtWidgets.QMessageBox.question(
+            self.window,
+            'Delete Config',
             f'Are you sure you would like to delete config "{rel_config_path}"?',
             QtWidgets.QMessageBox.Yes, QtWidgets.QMessageBox.No)
 
@@ -1143,7 +1148,7 @@ class SwitchboardDialog(QtCore.QObject):
             device_instances = self.device_manager.devices_of_type(plugin_name)
             device_settings = [(device.name, device.device_settings(), device.setting_overrides()) for device in device_instances]
             settings_dialog.add_section_for_plugin(plugin_name, self.device_manager.plugin_settings(plugin_name), device_settings)
-        
+
         settings_dialog.select_all_tab()
 
         old_address = SETTINGS.ADDRESS.get_value()
@@ -1188,7 +1193,7 @@ class SwitchboardDialog(QtCore.QObject):
     def sync_and_build_all_button_clicked(self):
         if not CONFIG.P4_ENABLED.get_value():
             return
-        
+
         self._sync_build_all(do_sync=True, do_build=True)
 
     def _sync_build_all(self, *, do_sync: bool, do_build: bool):
@@ -1351,7 +1356,7 @@ class SwitchboardDialog(QtCore.QObject):
         device_widget.signal_device_widget_trigger_start_toggled.connect(self.device_widget_trigger_start_toggled)
         device_widget.signal_device_widget_trigger_stop_toggled.connect(self.device_widget_trigger_stop_toggled)
 
-        #Live Link Face Toggle Display
+        # Live Link Face Toggle Display
         try:
             device_widget.signal_device_widget_toggle_display.connect(self.device_widget_toggle_display)
         except:
@@ -1418,7 +1423,7 @@ class SwitchboardDialog(QtCore.QObject):
 
         if self.window.multiuser_session_lineEdit.text() != value:
             self.window.multiuser_session_lineEdit.setText(value)
-        
+
         if value != SETTINGS.MUSERVER_SESSION_NAME:
             SETTINGS.MUSERVER_SESSION_NAME = value
             SETTINGS.save()
@@ -1512,7 +1517,7 @@ class SwitchboardDialog(QtCore.QObject):
         # Check is that slate/take combo has been used before
         while not self.recording_manager.slate_take_available(self._sequence, self._slate, requested_take):
             requested_take += 1
-        
+
         if requested_take == value == self._take:
             return
 
@@ -1520,7 +1525,6 @@ class SwitchboardDialog(QtCore.QObject):
             LOGGER.warning(f'Slate: "{self._slate}" Take: "{value}" have already been used. Auto incremented up to take: "{requested_take}"')
             # Clear the exclude list since Switchboard changed the incoming value
             exclude_addresses = []
-        
 
         self._take = requested_take
         SETTINGS.CURRENT_TAKE = value
@@ -1566,7 +1570,7 @@ class SwitchboardDialog(QtCore.QObject):
     def _on_selected_level_changed(self, index: int):
         full_map_path = self._get_level_from_combo_box(index)
         self._set_level(full_map_path)
-        
+
     def _get_level_from_combo_box(self, index: int):
         # Data stores full path
         return self.level_combo_box.itemData(index)
@@ -1738,7 +1742,7 @@ class SwitchboardDialog(QtCore.QObject):
             required_cl=self.project_changelist if self.project_changelist != EMPTY_SYNC_ENTRY else None,
             current_device_cl=device.project_changelist
         )
-        
+
         cl = device.project_changelist
         address = device.address
         for device in self.device_manager.devices():
@@ -1942,7 +1946,7 @@ class SwitchboardDialog(QtCore.QObject):
             LOGGER.info("Refreshing p4 project changelists")
             working_dir = os.path.dirname(CONFIG.UPROJECT_PATH.get_value())
             changelists = p4_utils.p4_latest_changelist(CONFIG.P4_PROJECT_PATH.get_value(), working_dir)
-            
+
         self.window.project_cl_combo_box.clear()
 
         if changelists:
@@ -1985,16 +1989,33 @@ class SwitchboardDialog(QtCore.QObject):
         self.window.engine_cl_combo_box.addItem(EMPTY_SYNC_ENTRY)
 
     def refresh_levels(self, levels: Optional[List[str]] = None):
+
         if levels is None:
             levels = CONFIG.maps()
 
+            # Update and save the CONFIG with the updated list of levels
+            CONFIG.LEVELS = levels
+            CONFIG.save()
+
         current_level = CONFIG.CURRENT_LEVEL
 
-        self.level_combo_box.clear()
-        self._update_level_list(self.level_combo_box, levels)
+        # Disable CONFIG saving because _update_level_list will trigger saves when updating
+        # the levels combobox even when the level selection hasn't change
+        CONFIG.push_saving_allowed(False)
+
+        # try-finally ensures that pop_saving_allowed happens in the case of exceptionsS
+        try:
+            self.level_combo_box.clear()
+            self._update_level_list(self.level_combo_box, levels)
+
+            if current_level and current_level in levels:
+                self.level = current_level
+        finally:
+            CONFIG.pop_saving_allowed()
 
-        if current_level and current_level in levels:
-            self.level = current_level
+        # if the current level changed, save the config
+        if self.level != current_level:
+            CONFIG.save()
 
     def filter_empty_abiguated_path(path, file_name):
         path = path.removesuffix(file_name)
@@ -2004,15 +2025,15 @@ class SwitchboardDialog(QtCore.QObject):
             return f"{file_name} ({path})"
 
     def generate_short_map_path(path: str, file_name: str) -> str:
-            path = path.replace("/Game", "", 1)
-            return SwitchboardDialog.filter_empty_abiguated_path(path, file_name)
-        
+        path = path.replace("/Game", "", 1)
+        return SwitchboardDialog.filter_empty_abiguated_path(path, file_name)
+
     def generate_disambiguated_names(path_list, shortening_function):
         name_counts = {}
         for path in path_list:
             file_name = os.path.basename(path)
             name_counts[file_name] = name_counts.get(file_name, 0) + 1
-            
+
         # Show only level name if unique and show path behind to disambiguate duplicates
         short_name_list = []
         short_name_to_path = {}
@@ -2021,23 +2042,25 @@ class SwitchboardDialog(QtCore.QObject):
             short_name = file_name if name_counts[file_name] == 1 else shortening_function(path, file_name)
             short_name_list.append(short_name)
             short_name_to_path[short_name] = path
-            
+
         return short_name_list, short_name_to_path
 
     def _update_level_list(self, level_combo_box: sb_widgets.SearchableComboBox, level_path_list: List[str]):
         def compare_file_names(path_a: str, path_b: str):
             return -1 if path_a.lower() < path_b.lower() \
                 else 1 if path_a.lower() > path_b.lower() else 0
-        
-        short_name_list, short_name_to_path = SwitchboardDialog.generate_disambiguated_names(level_path_list, SwitchboardDialog.generate_short_map_path)
-            
+
+        short_name_list, short_name_to_path = SwitchboardDialog.generate_disambiguated_names(
+            level_path_list,
+            SwitchboardDialog.generate_short_map_path)
+
         from functools import cmp_to_key
         short_name_list = sorted(short_name_list, key=cmp_to_key(compare_file_names))
         level_combo_box.addItems([DEFAULT_MAP_TEXT] + short_name_list)
-        
+
         if len(level_path_list) == 0:
             return
-        
+
         # To disambiguate, show the full path name in the drop-down
         level_combo_box.setItemData(0, DEFAULT_MAP_TEXT)
         level_combo_box.setItemData(0, "Default level", QtCore.Qt.ToolTipRole)
@@ -2079,10 +2102,7 @@ class SwitchboardDialog(QtCore.QObject):
             LOGGER.warning(f'Selected level "{current_level}" was removed; reverting to default')
             self.level = DEFAULT_MAP_TEXT
 
-        CONFIG.push_saving_allowed(False)
         self.refresh_levels(updated_levels_list)
-        CONFIG.pop_saving_allowed()
-
 
     def toggle_p4_controls(self, enabled):
         self.window.engine_cl_label.setEnabled(enabled)
@@ -2327,7 +2347,7 @@ class SwitchboardDialog(QtCore.QObject):
     def osc_battery(self, address, command, value):
         # The Battery command is used to handshake with LiveLinkFace. Don't reject it if it's not connected 
         device = self.device_manager.device_with_address(address[0])
-        
+
         if not device:
             return
 
@@ -2361,7 +2381,7 @@ class TransportQueueHeaderActionWidget(QtWidgets.QWidget):
         self.name_label = __label('Transport Queue')
 
         self.layout.addWidget(self.name_label)
-    
+
     def paintEvent(self, event):
         opt = QtWidgets.QStyleOption()
         opt.initFrom(self)
@@ -2407,7 +2427,7 @@ class TransportQueueActionWidget(QtWidgets.QWidget):
 
     def test(self):
         LOGGER.debug('BOOM!')
-    
+
     def paintEvent(self, event):
         opt = QtWidgets.QStyleOption()
         opt.initFrom(self)
