Switchboard: Added local cache that can be shared between configs.

It caches:
* Maps in a project
* Asset types of interest (LiveLink presets, media profiles, nDisplay configs)

Also added a progress dialog when finding the maps of a project, and an additional dialog when finding the assets of interest to parse.

A future improvement can be to do these tasks in the background, have a small status bar indicator of the progress, and update the UI when it is done.

diff --git a/Engine/Plugins/VirtualProduction/Switchboard/Source/Switchboard/switchboard/config.py b/Engine/Plugins/VirtualProduction/Switchboard/Source/Switchboard/switchboard/config.py
--- a/Engine/Plugins/VirtualProduction/Switchboard/Source/Switchboard/switchboard/config.py
+++ b/Engine/Plugins/VirtualProduction/Switchboard/Source/Switchboard/switchboard/config.py
@@ -9,6 +9,8 @@ import pathlib
 import shutil
 import socket
 import sys
+import threading
+import time
 from typing import Any, Callable, Optional, Tuple, Type, Union
 from enum import Enum
 
@@ -21,7 +23,7 @@ from switchboard.switchboard_logging import LOGGER
 from switchboard.switchboard_widgets import (
     DropDownMenuComboBox, NonScrollableComboBox)
 from switchboard import ue_plugin_utils, ugs_utils
-
+from switchboard.sbcache import SBCache, Map
 
 ROOT_CONFIGS_PATH = pathlib.Path(__file__).parent.with_name('configs')
 CONFIG_SUFFIX = '.json'
@@ -456,7 +458,7 @@ class Setting(QtCore.QObject):
             if override_device_name is None:
                 self._base_widget = None
             else:
-                del self._override_widgets[override_device_name]
+                self._override_widgets.pop(override_device_name, None)
         else:
             if override_device_name is None:
                 self._base_widget = widget
@@ -2091,7 +2093,6 @@ class Config(object):
 
         # MISC SETTINGS
         self.CURRENT_LEVEL = data.get('current_level', DEFAULT_MAP_TEXT)
-        self.LEVELS = data.get('levels', None)
 
         # Devices
         self._device_data_from_config = {}
@@ -2159,7 +2160,6 @@ class Config(object):
         self.init_muserver()
 
         self.CURRENT_LEVEL = DEFAULT_MAP_TEXT
-        self.LEVELS = None
 
         self._device_data_from_config = {}
         self._plugin_data_from_config = {}
@@ -2271,8 +2271,15 @@ class Config(object):
         elif data.get('build_engine', False):
             self.basic_project_settings["engine_sync_method"].update_value(EngineSyncMethod.Build_Engine.value)
 
-        self.PROJECT_NAME = self.basic_project_settings["project_name"]
         self.UPROJECT_PATH = self.basic_project_settings["uproject"]
+
+        # Take note if this project had a cache when opened, as devices may want to trigger a cache
+        if SBCache().query_project(self.UPROJECT_PATH.get_value()):
+            self.PROJECTWASINCACHE = True
+        else:
+            self.PROJECTWASINCACHE = False
+
+        self.PROJECT_NAME = self.basic_project_settings["project_name"]
         self.ENGINE_DIR = self.basic_project_settings["engine_dir"]
         self.ENGINE_SYNC_METHOD = self.basic_project_settings["engine_sync_method"]
         self.MAPS_PATH = self.basic_project_settings["maps_path"]
@@ -2577,9 +2584,8 @@ class Config(object):
 
         self.save_muserver(data)
 
-        # Levels
+        # Level
         data["current_level"] = self.CURRENT_LEVEL
-        data["levels"] = self.LEVELS
 
         # Devices
         data["devices"] = {}
@@ -2722,9 +2728,47 @@ class Config(object):
 
         return path_name
 
-    def maps(self):
+    def find_levels(self) -> list[str]:
+
+        # show a progress bar if it is taking more a trivial amount of time
+        progressDiag = QtWidgets.QProgressDialog(
+            'Finding levels...', 'Cancel', 0, 0, parent=None)
+
+        progressDiag.setWindowTitle('Unreal Level Finder')
+        progressDiag.setModal(True)
+        progressDiag.setMinimumDuration(1000)  # time before it shows up
+        progressDiag.setCancelButton(None)
+
+        # Looks much better without the window frame.
+        progressDiag.setWindowFlag(QtCore.Qt.FramelessWindowHint)
+
+        # create an event object to signal when the function is done
+        done_event = threading.Event()
+        levels = []
+
+        # our convenience worker function
+        def find_level_work():
+            self._find_levels(levels)
+            done_event.set()
+
+        thread = threading.Thread(target=find_level_work)
+        thread.start()
+
+        # wait for the event to be set or the progress dialog to be canceled
+        while not done_event.is_set() and not progressDiag.wasCanceled():
+            progressDiag.setValue(progressDiag.value() + 1)
+            QtWidgets.QApplication.processEvents()
+            time.sleep(0.050)  # The worker thread will run faster if we sleep.
+
+        progressDiag.close()
+
+        thread.join()
+
+        return levels
+
+    def _find_levels(self, levels: list[str]) -> None:
         '''
-        Returns a list of full map paths in an Unreal Engine project and
+        Returns a list of full level paths in an Unreal Engine project and
         in plugins such as:
             [
                 "/Game/Maps/MapName",
@@ -2771,12 +2815,15 @@ class Config(object):
                     umap_files.extend(get_umap_files(entry.path))  # recursive
             return umap_files
 
-        maps = []
         maps_filter = self.MAPS_FILTER.get_value()
 
         for (unreal_content_plugin, maps_path) in search_paths:
 
-            umaps = get_umap_files(maps_path)
+            try:
+                umaps = get_umap_files(maps_path)
+            except FileNotFoundError:
+                LOGGER.error(f"Could not get levels from {maps_path}")
+                continue
 
             for umap in umaps:
                 if not fnmatch.fnmatch(umap.name, maps_filter):
@@ -2789,11 +2836,21 @@ class Config(object):
                     file_path_to_map,
                     unreal_content_plugin=unreal_content_plugin)
 
-                if content_path_to_map not in maps:
-                    maps.append(content_path_to_map)
+                if content_path_to_map not in levels:
+                    levels.append(content_path_to_map)
+
+        levels.sort()
+
+        # cache updated list of maps
+        project = SBCache().query_or_create_project(self.UPROJECT_PATH.get_value())
+        cmaps = [Map(id=None, project=None, gamepath=map) for map in levels]
+        SBCache().update_project_maps(project=project, maps=cmaps)
 
-        maps.sort()
-        return maps
+    def get_levels(self) -> list[str]:
+        ''' Returns a list with all *.map files from the cache '''
+        project = SBCache().query_or_create_project(self.UPROJECT_PATH.get_value())
+        maps = SBCache().query_maps(project)
+        return [map.gamepath for map in maps]
 
     def multiuser_server_path(self):
         if self.MUSERVER_SLATE_MODE.get_value():
diff --git a/Engine/Plugins/VirtualProduction/Switchboard/Source/Switchboard/switchboard/devices/ndisplay/plugin_ndisplay.py b/Engine/Plugins/VirtualProduction/Switchboard/Source/Switchboard/switchboard/devices/ndisplay/plugin_ndisplay.py
--- a/Engine/Plugins/VirtualProduction/Switchboard/Source/Switchboard/switchboard/devices/ndisplay/plugin_ndisplay.py
+++ b/Engine/Plugins/VirtualProduction/Switchboard/Source/Switchboard/switchboard/devices/ndisplay/plugin_ndisplay.py
@@ -26,6 +26,7 @@ from switchboard.devices.unreal.plugin_unreal import DeviceUnreal, \
 from switchboard.devices.unreal.uassetparser import UassetParser
 from switchboard.devices.device_base import DeviceStatus
 from switchboard.switchboard_logging import LOGGER
+from switchboard.sbcache import SBCache, Asset
 
 from .ndisplay_monitor_ui import nDisplayMonitorUI
 from .ndisplay_monitor import nDisplayMonitor
@@ -99,45 +100,36 @@ class AddnDisplayDialog(AddDeviceDialog):
         having to re-find every time.
         '''
 
-         # cache the current value of the combo box (currently selected preset for this device)
+        # cache the current value of the combo box (currently selected preset for this device)
         cur_item = self.cbConfigs.currentData()
 
         self.cbConfigs.clear()
 
-        itemDatas = DeviceUnreal.csettings['asset_itemDatas'].get_value()
+        project = SBCache().query_or_create_project(CONFIG.UPROJECT_PATH.get_value())
+        assets = SBCache().query_assets_by_classname(project=project, classnames=DeviceUnreal.NDISPLAY_CLASS_NAMES)
 
-        def asset_is_nDisplay_config(itemData:dict) -> bool:
-            ''' Convenience function to filter the nDisplay configs '''
-
-            # Only return assets of the correct class
-            return itemData['classname'] in DeviceUnreal.NDISPLAY_CLASS_NAMES
-
-        try:
-            for itemData in [itemData for itemData in itemDatas if asset_is_nDisplay_config(itemData)]:
-                self.cbConfigs.addItem(itemData['name'], itemData)
-        except Exception:
-            LOGGER.error('Error recalling config itemDatas')
+        for asset in assets:
+            self.cbConfigs.addItem(asset.name, asset)
 
         # restore selected item
         for item_idx in range(self.cbConfigs.count()):
-            if cur_item and (cur_item['name'] == self.cbConfigs.itemData(item_idx)['name']):
+            if cur_item and (cur_item.name == self.cbConfigs.itemData(item_idx).name):
                 self.cbConfigs.setCurrentIndex(item_idx)
                 break
 
-
     def current_config_path(self):
         ''' Get currently selected config path in the combobox '''
 
         # Detect if this is a manual entry using the itemData
 
         itemText = self.cbConfigs.currentText()
-        itemData = self.cbConfigs.currentData()
+        asset: Asset = self.cbConfigs.currentData()
 
         if (self.cbConfigs.currentData() is None or
-                itemData['name'] != itemText):
+                asset.name != itemText):
             config_path = itemText.replace('"', '').strip()
         else:
-            config_path = itemData['path'].replace('"', '').strip()
+            config_path = asset.localpath.replace('"', '').strip()
 
         # normalize the path
         config_path = os.path.normpath(config_path)
diff --git a/Engine/Plugins/VirtualProduction/Switchboard/Source/Switchboard/switchboard/devices/unreal/plugin_unreal.py b/Engine/Plugins/VirtualProduction/Switchboard/Source/Switchboard/switchboard/devices/unreal/plugin_unreal.py
--- a/Engine/Plugins/VirtualProduction/Switchboard/Source/Switchboard/switchboard/devices/unreal/plugin_unreal.py
+++ b/Engine/Plugins/VirtualProduction/Switchboard/Source/Switchboard/switchboard/devices/unreal/plugin_unreal.py
@@ -18,6 +18,7 @@ import sys
 import threading
 from typing import Callable, Generator, Optional, Union
 import uuid
+import time
 
 from PySide2 import QtCore, QtGui, QtWidgets
 
@@ -39,6 +40,7 @@ from switchboard.util import p4_changelist_inspection
 from . import version_helpers
 from .listener_watcher import ListenerWatcher
 from .redeploy_dialog import RedeployListenerDialog
+from switchboard.sbcache import SBCache, Asset
 
 
 class ProgramStartQueueItem:
@@ -294,7 +296,7 @@ class LiveLinkPresetSetting(Setting):
             allow_reset=allow_reset,
             migrate_data=migrate_data)
 
-    def _create_widgets(self, override_device_name = None):
+    def _create_widgets(self, override_device_name=None):
 
         # create combo with livelink preset options
 
@@ -340,31 +342,16 @@ class LiveLinkPresetSetting(Setting):
 
     def _validate_and_commit_value(self, combo: QtWidgets.QComboBox, override_device_name:Optional[str] = None):
 
-        itemData = combo.currentData()
+        asset: Asset = combo.currentData()
 
         value_str = ''
 
-        if itemData:
-            value_str = itemData['gamepath'] # we use path and not name because it is unambiguous
+        if asset:
+            value_str = asset.gamepath  # we use path and not name because it is unambiguous
 
         self._on_widget_value_changed(value_str, override_device_name=override_device_name)
 
-
-    def asset_is_relevant(self, itemData) -> bool:
-        ''' Convenience function to filter the relevant assets for this setting '''
-
-        # Only return assets of the correct live link preset class
-        if itemData['classname'] not in DeviceUnreal.LIVELINKPRESET_CLASS_NAMES:
-            return False
-            
-        # Only return live link presets that have a valid gamepath. 
-        # We wouldn't be able to build the command line otherwise
-        if itemData['gamepath'] == '':
-            return False
-
-        return True
-
-    def _update_combo_items(self, combo:QtWidgets.QComboBox, override_device_name):
+    def _update_combo_items(self, combo: QtWidgets.QComboBox, override_device_name):
         '''
         Populate the combobox itself with known list of live link presets available.
         Makes sure that the currently selected preset is preserved, unless it is not found,
@@ -378,31 +365,28 @@ class LiveLinkPresetSetting(Setting):
         combo.clear()
 
         # add the empty/none choice
-        combo.addItem('', {'name':'', 'gamepath':''})
+        noneasset = Asset(id=0, project=None, assettype=None, gamepath='', name='', localpath='')
+        combo.addItem(noneasset.name, noneasset)
 
-        itemDatas = DeviceUnreal.csettings['asset_itemDatas'].get_value()
+        project = SBCache().query_or_create_project(CONFIG.UPROJECT_PATH.get_value())
+        assets = SBCache().query_assets_by_classname(
+            project=project,
+            classnames=self._classnames())
 
         # generate the combo box items
-        try:
-            for itemData in [itemData for itemData in itemDatas if self.asset_is_relevant(itemData)]:
-                name = itemData['name']
-
-                # trim the expected .uasset extension
-                ext = '.uasset'
-                if name.endswith(ext):
-                    name = name[:len(name)-len(ext)]
-
-                combo.addItem(name, itemData)
-
-        except Exception as e:
-            LOGGER.error(f'Error recalling asset itemDatas: {e}')
+        for asset in assets:
+            combo.addItem(asset.name.replace('.uasset', ''), asset)
 
         # set the current index to the live link preset that was already selected
         for item_idx in range(combo.count()):
-            if cur_value == combo.itemData(item_idx)['gamepath']:
+            if cur_value == combo.itemData(item_idx).gamepath:
                 combo.setCurrentIndex(item_idx)
                 break
 
+    def _classnames(self) -> list[str]:
+        ''' Return the valid class names of this type of asset '''
+        return DeviceUnreal.LIVELINKPRESET_CLASS_NAMES
+
     def _on_setting_changed(self, new_value: str, override_device_name: Optional[str] = None):
 
         combo = self.get_widget(override_device_name=override_device_name)
@@ -411,8 +395,9 @@ class LiveLinkPresetSetting(Setting):
             return
 
         try:
-            old_value = combo.currentData()['gamepath']
-        except (KeyError, TypeError):
+            asset: Asset = combo.currentData()
+            old_value = asset.gamepath
+        except (AttributeError, TypeError):
             old_value = ''
 
         new_str_value = new_value
@@ -420,7 +405,8 @@ class LiveLinkPresetSetting(Setting):
         # if the value changed, find the new index in the combo box based on gamepath
         if new_str_value != old_value:
             for item_idx in range(combo.count()):
-                if new_str_value == combo.itemData(item_idx)['gamepath']:
+                asset: Asset = combo.itemData(item_idx)
+                if new_str_value == asset.gamepath:
                     combo.setCurrentIndex(item_idx)
                     break
 
@@ -428,19 +414,9 @@ class LiveLinkPresetSetting(Setting):
 class MediaProfileSetting(LiveLinkPresetSetting):
     ''' Container of the MediaProfile setting.
     '''
-
-    #@override
-    def asset_is_relevant(self, itemData) -> bool:
-        # Only return assets of the correct live link preset class
-        if itemData['classname'] not in DeviceUnreal.MEDIAPROFILE_CLASS_NAMES:
-            return False
-            
-        # Only return live link presets that have a valid gamepath. 
-        # We wouldn't be able to build the command line otherwise
-        if itemData['gamepath'] == '':
-            return False
-
-        return True
+    def _classnames(self) -> list[str]:
+        ''' Return the valid class names of this type of asset '''
+        return DeviceUnreal.MEDIAPROFILE_CLASS_NAMES
 
 
 class SyncCategoryOption:
@@ -654,13 +630,6 @@ class DeviceUnreal(Device):
             tool_tip=(
                 'Adds the selected LiveLink preset to the command line \n')
         ),
-        'asset_itemDatas': Setting(
-            attr_name='asset_itemDatas',
-            nice_name="Asset files",
-            value=[],
-            tool_tip="Remember the last analyzed list of project and plugin assets",
-            show_ui=False,
-        ),
         'mediaprofile': MediaProfileSetting(
             attr_name='mediaprofile',
             nice_name='Media Profile',
@@ -2866,10 +2835,22 @@ class DeviceUnreal(Device):
         QtGui.QGuiApplication.clipboard().setText(
             self.last_launch_command.get_value())
 
+    @classmethod
+    def all_devices_added(cls):
+        ''' Device interface implementation '''
+
+        super().all_devices_added()
+
+        # Trigger a project cache if this is a DeviceUnreal class (not a subclass)
+        # and the project was not cached when opened
+        if cls == DeviceUnreal:
+            if not CONFIG.PROJECTWASINCACHE:
+                DeviceUnreal.analyze_project_assets()
+
     @classmethod
     def analyze_project_assets(cls):
-        ''' Traverses project and content plugins and caches a list of assets of interest. 
-        That list is currently nDisplay configs and live link presets.
+        ''' Traverses project and content plugins and caches a list of assets of interest.
+        e.g. nDisplay configs, live link presets, and media profiles.
         '''
 
         project_configs_path = os.path.normpath(
@@ -2893,29 +2874,62 @@ class DeviceUnreal(Device):
 
         assets = []
 
-        for (unreal_content_plugin, configs_path) in search_paths:
-            for dirpath, _, file_names in os.walk(configs_path):
-                for file_name in file_names:
-                    if not file_name.lower().endswith(('.uasset', '.ndisplay')):
-                        continue
-
-                    if file_name not in asset_names:
-                        asset_path = os.path.join(dirpath, file_name)
-                        ext = os.path.splitext(file_name)[1]
-
-                        # Since .uasset is a generic asset container, only add
-                        # assets of the right class.
-                        if ext.lower() == '.uasset':
-                            assets.append({
-                                'name': file_name,
-                                'path': asset_path,
-                                'plugin': unreal_content_plugin,
-                            })
-                        else:
-                            asset_names.append(file_name)
-                            asset_paths.append(asset_path)
-                            asset_plugins.append(unreal_content_plugin)
-                            asset_classnames.append(DeviceUnreal.NDISPLAY_CLASS_NAMES[0]) # so that it passes the filter later on
+        # show a progress bar if it is taking more a trivial amount of time
+        progressDiag = QtWidgets.QProgressDialog(
+            'Finding assets...', 'Cancel', 0, 0, parent=None)
+
+        progressDiag.setWindowTitle('Unreal Asset Finder')
+        progressDiag.setModal(True)
+        progressDiag.setMinimumDuration(1000)  # time before it shows up
+        progressDiag.setCancelButton(None)
+
+        # Looks much better without the window frame.
+        progressDiag.setWindowFlag(QtCore.Qt.FramelessWindowHint)
+
+        # create an event object to signal when the function is done
+        done_event = threading.Event()
+
+        # convenience find assets worker. This is a step before parsing the assets, and can
+        # also take some time depending on the number of assets in the project.
+        def find_assets_work():
+            for (unreal_content_plugin, configs_path) in search_paths:
+                for dirpath, _, file_names in os.walk(configs_path):
+                    for file_name in file_names:
+                        if not file_name.lower().endswith(('.uasset', '.ndisplay')):
+                            continue
+
+                        if file_name not in asset_names:
+                            asset_path = os.path.join(dirpath, file_name)
+                            ext = os.path.splitext(file_name)[1]
+
+                            # Since .uasset is a generic asset container, only add
+                            # assets of the right class.
+                            if ext.lower() == '.uasset':
+                                assets.append({
+                                    'name': file_name,
+                                    'path': asset_path,
+                                    'plugin': unreal_content_plugin,
+                                })
+                            else:
+                                asset_names.append(file_name)
+                                asset_paths.append(asset_path)
+                                asset_plugins.append(unreal_content_plugin)
+                                asset_classnames.append(DeviceUnreal.NDISPLAY_CLASS_NAMES[0]) # so that it passes the filter later on
+
+            done_event.set()
+
+        thread = threading.Thread(target=find_assets_work)
+        thread.start()
+
+        # wait for the event to be set or the progress dialog to be canceled
+        while not done_event.is_set() and not progressDiag.wasCanceled():
+            progressDiag.setValue(progressDiag.value() + 1)
+            QtWidgets.QApplication.processEvents()
+            time.sleep(0.050)  # The worker thread will run faster if we sleep.
+
+        progressDiag.close()
+
+        thread.join()
 
         # process the assets in a multi-threaded fashion
 
@@ -2936,7 +2950,7 @@ class DeviceUnreal(Device):
             DeviceUnreal.NDISPLAY_CLASS_NAMES \
             + DeviceUnreal.LIVELINKPRESET_CLASS_NAMES \
             + DeviceUnreal.MEDIAPROFILE_CLASS_NAMES
-        
+
         def validateInterestingAsset(asset):                                                                                                
             ''' Returns the asset if it is an interesting asset '''
 
@@ -2982,7 +2996,7 @@ class DeviceUnreal(Device):
                     asset_paths.append(path)
                     asset_plugins.append(plugin)
                     asset_classnames.append(classname)
-                    
+
                 except Exception:
                     pass
 
@@ -2993,36 +3007,35 @@ class DeviceUnreal(Device):
             config_path = CONFIG.shrink_path(config_path)
             return sb_dialog.SwitchboardDialog.filter_empty_abiguated_path(config_path, file_name)
 
-        asset_names, _ = sb_dialog.SwitchboardDialog.generate_disambiguated_names(asset_paths, generate_short_unique_config_name)
+        asset_names, _ = sb_dialog.SwitchboardDialog.generate_disambiguated_names(
+            asset_paths, generate_short_unique_config_name)
+
+        # collect the found config files into the assets list
 
-        # collect the found config files into the itemDatas list
+        assets = []
 
-        itemDatas = []
+        project = SBCache().query_or_create_project(CONFIG.UPROJECT_PATH.get_value())
 
         for idx, asset_name in enumerate(asset_names):
-            
-            uplugin_file_path = (
-                str(asset_plugins[idx].uplugin_file_path) if asset_plugins[idx]
-                else None)
 
             gamepath = CONFIG.resolve_content_path(
                 file_path=asset_paths[idx], 
                 unreal_content_plugin=asset_plugins[idx])
 
-            itemData = {
-                'name': asset_name,
-                'path': asset_paths[idx],
-                'uplugin_file_path': uplugin_file_path,
-                'classname': asset_classnames[idx],
-                'gamepath': gamepath,
-            }
-            itemDatas.append(itemData)
+            assettype = SBCache().query_or_create_assettype(asset_classnames[idx])
+            asset = Asset(
+                id=0,
+                project=project,
+                assettype=assettype,
+                gamepath=gamepath,
+                name=asset_name,
+                localpath=asset_paths[idx])
+
+            assets.append(asset)
 
-        # sort by name
-        itemDatas.sort(key=lambda itemData: itemData['name'])
+        # save to the cache
+        SBCache().update_project_assets(project=project, assets=assets)
 
-        # update settings that should survive device removal and addition
-        DeviceUnreal.csettings['asset_itemDatas'].update_value(itemDatas)
 
 def parse_unreal_tag_file(file_content):
     tags = []
diff --git a/Engine/Plugins/VirtualProduction/Switchboard/Source/Switchboard/switchboard/sbcache.py b/Engine/Plugins/VirtualProduction/Switchboard/Source/Switchboard/switchboard/sbcache.py
new file mode 100644
--- /dev/null
+++ b/Engine/Plugins/VirtualProduction/Switchboard/Source/Switchboard/switchboard/sbcache.py
@@ -0,0 +1,359 @@
+# Copyright Epic Games, Inc. All Rights Reserved.
+
+import sqlite3
+import pathlib
+
+
+ROOT_SBCACHE_PATH = pathlib.Path(__file__).parent.with_name('cache')
+
+
+class Row:
+    '''Base class for database rows in a table'''
+    def __init__(self, id: int):
+        self.id = id
+
+
+class Project(Row):
+    '''Represents a row in the Projects table'''
+    def __init__(self, id: int, path: str):
+        super().__init__(id)
+        self.path = path
+
+
+class Map(Row):
+    '''Represents a row in the Maps table. Every map belongs to a project.'''
+    def __init__(self, id: int, project: Project, gamepath: str):
+        super().__init__(id)
+        self.project = project
+        self.gamepath = gamepath
+
+
+class AssetType(Row):
+    '''Represents a row in the AssetTypes table. Used to identify the type of Assets.
+    The type refers to the class name of the asset.
+    '''
+    def __init__(self, id: int, classname: str):
+        super().__init__(id)
+        self.classname = classname
+
+
+class Asset(Row):
+    '''Represents a row in the Assets table. They all belong to a project, and have an AssetType.'''
+    def __init__(
+            self,
+            id: int,
+            project: Project,
+            assettype: AssetType,
+            gamepath: str,
+            name: str,
+            localpath: str):
+
+        super().__init__(id)
+
+        self.project = project
+        self.gamepath = gamepath
+        self.assettype = assettype
+        self.name = name
+        self.localpath = localpath
+
+
+class SBCache:
+    ''' Used to keep a local cache of data gleaned from projects but
+    that can be shared between switchboard configs. Its purpose is to
+    avoid repeated project content searches and ultimately speed up
+    opening configs. It also helps keep the projects configs small.
+    For example:
+        * Maps in a project.
+        * LiveLink presets and media profiles in a project.
+
+    Schema:
+        * Projects
+            * ProjectId (primary key)
+            * LocalPath
+        * Maps
+            * MapId  (primary key)
+            * ProjectId
+            * GamePath
+        * AssetTypes
+            * AssetTypeId (primary key)
+            * ClassName (unique)
+        * Assets
+            * AssetId (primary key)
+            * ProjectId
+            * AssetTypeId
+            * Name
+            * GamePath
+            * LocalPath
+    '''
+
+    version = (1, 0, 0)  # Major, Minor, Patch. Change when schema is updated.
+
+    def __new__(cls):
+        ''' Classic singleton pattern '''
+        if not hasattr(cls, 'instance'):
+            cls.instance = super(SBCache, cls).__new__(cls)
+        return cls.instance
+
+    def __init__(self):
+
+        self.con = None
+
+        try:
+            self._connect()
+            self._verifydb()
+        except (sqlite3.DatabaseError, sqlite3.Error, ValueError):
+            self._recreatedb()
+            self._verifydb()
+
+    def _recreatedb(self):
+        ''' Deletes existing file and re-creates the database '''
+
+        if self.con:
+            self.con.close()
+
+        try:
+            self._dbfilepath().unlink()
+        except FileNotFoundError:
+            pass
+
+        self._dbfilepath().parent.mkdir(parents=True, exist_ok=True)
+        self._connect()
+        self._initializedb()
+
+    def _connect(self):
+        ''' Connect to the database in the given path '''
+        self.con = sqlite3.connect(str(self._dbfilepath()))
+
+    def _initializedb(self):
+        ''' Creates required tables '''
+
+        # Write metadata
+        self.con.execute('''
+            CREATE TABLE IF NOT EXISTS Metadata (
+                Major INTEGER NOT NULL,
+                Minor INTEGER NOT NULL,
+                Patch INTEGER NOT NULL
+            )
+            ''')
+
+        # We will use the version to invalidate the cache
+        self.con.execute('''
+            INSERT INTO Metadata (Major, Minor, Patch)
+            VALUES (?, ?, ?)
+        ''', self.version)
+
+        self.con.execute('''
+            CREATE TABLE IF NOT EXISTS Projects (
+                ProjectId INTEGER PRIMARY KEY,
+                LocalPath TEXT NOT NULL
+            )
+        ''')
+
+        self.con.execute('''
+            CREATE TABLE IF NOT EXISTS Maps (
+                MapId INTEGER PRIMARY KEY,
+                GamePath TEXT NOT NULL,
+                ProjectId INTEGER NOT NULL,
+                FOREIGN KEY (ProjectId) REFERENCES Projects (ProjectId)
+            )
+        ''')
+
+        # Create table for asset types
+        self.con.execute('''
+            CREATE TABLE IF NOT EXISTS AssetTypes (
+                AssetTypeId INTEGER PRIMARY KEY,
+                ClassName TEXT NOT NULL UNIQUE
+            )
+        ''')
+
+        # Create table for assets
+        self.con.execute('''
+            CREATE TABLE IF NOT EXISTS Assets (
+                AssetId INTEGER PRIMARY KEY,
+                ProjectId INTEGER NOT NULL,
+                AssetTypeId INTEGER NOT NULL,
+                Name TEXT NOT NULL,
+                GamePath TEXT NOT NULL,
+                LocalPath TEXT NOT NULL,
+                FOREIGN KEY (ProjectId) REFERENCES Projects (ProjectId),
+                FOREIGN KEY (AssetTypeId) REFERENCES AssetTypes (AssetTypeId)
+            )
+        ''')
+
+        self.con.commit()
+
+    def _verifydb(self):
+        ''' Makes sure that everything needed in the database exists '''
+
+        # check the version
+        cur = self.con.execute('''
+            SELECT Major, Minor, Patch FROM Metadata
+        ''')
+
+        row = cur.fetchone()
+
+        dbversion = (row[0], row[1], row[2])
+
+        if dbversion != self.version:
+            raise ValueError
+
+    def _dbfilename(self) -> str:
+        ''' Returns default filename for this database '''
+        return "sbcache.db"
+
+    def _dbfilepath(self) -> str:
+        ''' Returns default filepath of this database '''
+        return ROOT_SBCACHE_PATH / self._dbfilename()
+
+    def query_or_create_project(self, filepath: str) -> Project:
+        ''' Creates or returns the project with the given filepath'''
+
+        # Check if the project exists in the database
+        cur = self.con.execute('''
+            SELECT ProjectId, LocalPath FROM Projects
+            WHERE LocalPath = ?
+        ''', (filepath,))
+
+        row = cur.fetchone()
+
+        if row:
+            # Return the existing project
+            return Project(row[0], row[1])
+        else:
+            # Insert a new project and return it
+            cur = self.con.execute('''
+                INSERT INTO Projects (LocalPath)
+                VALUES (?)
+            ''', (filepath,))
+
+            self.con.commit()
+
+            return Project(id=cur.lastrowid, path=filepath)
+
+    def query_project(self, filepath: str) -> Project:
+        ''' Returns the project with the given filepath. None if non-existent'''
+
+        cur = self.con.execute('''
+            SELECT ProjectId, LocalPath FROM Projects
+            WHERE LocalPath = ?
+        ''', (filepath,))
+
+        row = cur.fetchone()
+
+        if row:
+            return Project(row[0], row[1])
+
+        return None
+
+    def query_maps(self, project: Project) -> list[Map]:
+        ''' Returns a list with all the maps of the given project '''
+
+        # Query the maps by the project id
+        cur = self.con.execute('''
+            SELECT MapId, GamePath FROM Maps
+            WHERE ProjectId = ?
+        ''', (project.id,))
+
+        rows = cur.fetchall()
+
+        # Convert the rows to Map objects
+        maps = [Map(id=row[0], project=project, gamepath=row[1]) for row in rows]
+        return maps
+
+    def update_project_maps(self, maps: list[Map], project: Project):
+        ''' Replaces all the maps of a project with those in the given list. '''
+
+        # Delete the existing maps of the project
+        self.con.execute('''
+            DELETE FROM Maps
+            WHERE ProjectId = ?
+        ''', (project.id,))
+
+        # Insert the new maps of the project
+        self.con.executemany('''
+            INSERT INTO Maps (GamePath, ProjectId)
+            VALUES (?, ?)
+        ''', [(map.gamepath, project.id) for map in maps])
+
+        self.con.commit()
+
+    def query_or_create_assettype(self, classname: str) -> AssetType:
+        ''' Returns the asset type of the given classname, which should be unique.
+        If it doesn't exist, then it creates it.
+        '''
+
+        # Check if the asset type exists in the database
+        cur = self.con.execute('''
+            SELECT AssetTypeId, ClassName FROM AssetTypes
+            WHERE ClassName = ?
+        ''', (classname,))
+
+        row = cur.fetchone()
+
+        if row:
+            # Return the existing asset type
+            return AssetType(row[0], row[1])
+        else:
+            # Insert a new asset type and return it
+            cur = self.con.execute('''
+                INSERT INTO AssetTypes (ClassName)
+                VALUES (?)
+            ''', (classname,))
+
+            self.con.commit()
+
+            return AssetType(id=cur.lastrowid, classname=classname)
+
+    def query_assets_by_classname(self, project: Project, classnames: list[str]) -> list[Asset]:
+        ''' Returns a list of assets that are of any of the given classnames. '''
+
+        # Convert the classnames to asset type ids
+        assettypes = [self.query_or_create_assettype(classname) for classname in classnames]
+        assettype_ids = [assettype.id for assettype in assettypes]
+
+        # Query the assets by the project id and the asset type ids
+        idstring = ','.join(['?'] * len(assettype_ids))
+        cur = self.con.execute(
+            f'''
+            SELECT AssetId, AssetTypeId, GamePath, Name, LocalPath FROM Assets
+            WHERE ProjectId = ? AND AssetTypeId IN ({idstring})
+            ''',
+            (project.id, *assettype_ids)
+        )
+
+        rows = cur.fetchall()
+
+        # Convert the rows to Asset objects
+
+        assetype_for_id = {assettype.id: assettype for assettype in assettypes}
+
+        def AssetFromRow(row):
+            assettype = assetype_for_id[row[1]]
+            return Asset(
+                id=row[0],
+                project=project,
+                assettype=assettype,
+                gamepath=row[2],
+                name=row[3],
+                localpath=row[4])
+
+        return [AssetFromRow(row) for row in rows]
+
+    def update_project_assets(self, project: Project, assets: list[Asset]) -> None:
+        ''' Updates all the assets of the given project. Since the project is specified in
+        the arguments, the project property of the assets in the list is ignored.
+        '''
+
+        # Delete the existing assets of the project
+        self.con.execute('''
+            DELETE FROM Assets
+            WHERE ProjectId = ?
+        ''', (project.id,))
+
+        # Insert the new assets of the project
+        self.con.executemany('''
+            INSERT INTO Assets (ProjectId, AssetTypeId, Name, GamePath, LocalPath)
+            VALUES (?, ?, ?, ?, ?)
+        ''', [(project.id, asset.assettype.id, asset.name, asset.gamepath, asset.localpath) for asset in assets])
+
+        self.con.commit()
diff --git a/Engine/Plugins/VirtualProduction/Switchboard/Source/Switchboard/switchboard/switchboard_dialog.py b/Engine/Plugins/VirtualProduction/Switchboard/Source/Switchboard/switchboard/switchboard_dialog.py
--- a/Engine/Plugins/VirtualProduction/Switchboard/Source/Switchboard/switchboard/switchboard_dialog.py
+++ b/Engine/Plugins/VirtualProduction/Switchboard/Source/Switchboard/switchboard/switchboard_dialog.py
@@ -486,7 +486,7 @@ class SwitchboardDialog(QtCore.QObject):
 
         if CONFIG.file_path:
             self.toggle_p4_controls(CONFIG.P4_ENABLED.get_value())
-            self.refresh_levels(CONFIG.LEVELS)
+            self.refresh_levels(CONFIG.get_levels())
         else:
             self.menu_new_config()
 
@@ -952,7 +952,7 @@ class SwitchboardDialog(QtCore.QObject):
 
         self.p4_refresh_project_cl()
         self.p4_refresh_engine_cl()
-        self.refresh_levels(levels=CONFIG.LEVELS)
+        self.refresh_levels(CONFIG.get_levels())
         self.update_current_config_text()
         self.refresh_muserver_autojoin()
         self.refresh_trace_settings()
@@ -1041,7 +1041,7 @@ class SwitchboardDialog(QtCore.QObject):
 
         # Update the UI
         self.toggle_p4_controls(CONFIG.P4_ENABLED.get_value())
-        self.refresh_levels()
+        self.refresh_levels(CONFIG.get_levels())
         self.update_current_config_text()
         self.refresh_muserver_autojoin()
         self.refresh_trace_settings()
@@ -1990,12 +1990,8 @@ class SwitchboardDialog(QtCore.QObject):
 
     def refresh_levels(self, levels: Optional[List[str]] = None):
 
-        if levels is None:
-            levels = CONFIG.maps()
-
-            # Update and save the CONFIG with the updated list of levels
-            CONFIG.LEVELS = levels
-            CONFIG.save()
+        if levels is None or len(levels) == 0:
+            levels = CONFIG.find_levels()
 
         current_level = CONFIG.CURRENT_LEVEL
 
@@ -2082,7 +2078,7 @@ class SwitchboardDialog(QtCore.QObject):
 
         current_level: str = CONFIG.CURRENT_LEVEL
         prev_levels_list: List[str] = self.get_current_level_list()
-        updated_levels_list: List[str] = CONFIG.maps()
+        updated_levels_list: List[str] = CONFIG.find_levels()
 
         prev_levels: Set[str] = set(prev_levels_list)
         updated_levels: Set[str] = set(updated_levels_list)
