Switchboard: Flag simultaneous multithreading as a warning in the CpuUtilization column of the nDisplay Monitor.

diff --git a/Engine/Plugins/VirtualProduction/Switchboard/Source/Switchboard/switchboard/devices/ndisplay/ndisplay_monitor.py b/Engine/Plugins/VirtualProduction/Switchboard/Source/Switchboard/switchboard/devices/ndisplay/ndisplay_monitor.py
--- a/Engine/Plugins/VirtualProduction/Switchboard/Source/Switchboard/switchboard/devices/ndisplay/ndisplay_monitor.py
+++ b/Engine/Plugins/VirtualProduction/Switchboard/Source/Switchboard/switchboard/devices/ndisplay/ndisplay_monitor.py
@@ -71,7 +71,10 @@ class nDisplayMonitor(QAbstractTableModel):
             'CpuUtilization':
                 'CPU utilization average. The number of overloaded cores (>'
                 f'{self.CORE_OVERLOAD_THRESH}% load) will be displayed in '
-                'parentheses.',
+                'parentheses.\n\n'
+                'in addition, simultaneous multi-threading (SMT), also known\n'
+                'as Hyper-Threading, is known to potentially cause hitches.\n'
+                'If SMT is enabled, that will also be flagged as a warning.',
             'MemUtilization': 'Physical memory, utilized / total.',
             'GpuUtilization':
                 'GPU utilization. The GPU clock speed is displayed in '
@@ -112,8 +115,9 @@ class nDisplayMonitor(QAbstractTableModel):
             return self.COLOR_NORMAL if is_normal else self.COLOR_WARNING
 
         if colname == 'CpuUtilization':
-            no_overload = '(' not in value  # "(# cores > threshold%)"
-            return self.COLOR_NORMAL if no_overload else self.COLOR_WARNING
+            # "(# cores > threshold%)" or "(SMT ENABLED)"
+            no_caveats = '(' not in value
+            return self.COLOR_NORMAL if no_caveats else self.COLOR_WARNING
 
         return self.COLOR_NORMAL
 
@@ -456,6 +460,9 @@ class nDisplayMonitor(QAbstractTableModel):
             if num_overloaded_cores > 0:
                 data['CpuUtilization'] += f' ({num_overloaded_cores} cores >' \
                     f' {self.CORE_OVERLOAD_THRESH}%)'
+
+            if device.processor_smt:
+                data['CpuUtilization'] += ' (SMT ENABLED)'
         except (KeyError, ValueError):
             data['CpuUtilization'] = self.DATA_MISSING
 
diff --git a/Engine/Plugins/VirtualProduction/Switchboard/Source/Switchboard/switchboard/devices/unreal/plugin_unreal.py b/Engine/Plugins/VirtualProduction/Switchboard/Source/Switchboard/switchboard/devices/unreal/plugin_unreal.py
--- a/Engine/Plugins/VirtualProduction/Switchboard/Source/Switchboard/switchboard/devices/unreal/plugin_unreal.py
+++ b/Engine/Plugins/VirtualProduction/Switchboard/Source/Switchboard/switchboard/devices/unreal/plugin_unreal.py
@@ -2577,6 +2577,7 @@ class DeviceUnreal(Device):
         self.os_version_label = message.get('osVersionLabel', '')
         self.os_version_label_sub = message.get('osVersionLabelSub', '')
         self.os_version_number = message.get('osVersionNumber', '')
+        self.processor_smt = message.get('bProcessorSMT', False)
         self.total_phys_mem = message.get('totalPhysicalMemory', 0)
         self.platform_binary_directory = message.get('platformBinaryDirectory', '')
 
diff --git a/Engine/Source/Programs/SwitchboardListener/SblCore/Internal/SwitchboardListener.h b/Engine/Source/Programs/SwitchboardListener/SblCore/Internal/SwitchboardListener.h
--- a/Engine/Source/Programs/SwitchboardListener/SblCore/Internal/SwitchboardListener.h
+++ b/Engine/Source/Programs/SwitchboardListener/SblCore/Internal/SwitchboardListener.h
@@ -282,6 +282,7 @@ private:
 	TSharedPtr<FSBLHelperClient> SBLHelper;
 
 	bool bIsNvAPIInitialized;
+	bool bProcessorSMT;
 
 	TSharedPtr<FRWLock> CachedMosaicToposLock;
 	TSharedPtr<TArray<struct FMosaicTopo>> CachedMosaicTopos;
diff --git a/Engine/Source/Programs/SwitchboardListener/SblCore/Internal/SwitchboardListenerVersion.h b/Engine/Source/Programs/SwitchboardListener/SblCore/Internal/SwitchboardListenerVersion.h
--- a/Engine/Source/Programs/SwitchboardListener/SblCore/Internal/SwitchboardListenerVersion.h
+++ b/Engine/Source/Programs/SwitchboardListener/SblCore/Internal/SwitchboardListenerVersion.h
@@ -4,4 +4,4 @@
 
 #define SBLISTENER_VERSION_MAJOR 3
 #define SBLISTENER_VERSION_MINOR 0
-#define SBLISTENER_VERSION_PATCH 2
+#define SBLISTENER_VERSION_PATCH 3
diff --git a/Engine/Source/Programs/SwitchboardListener/SblCore/Private/SwitchboardListener.cpp b/Engine/Source/Programs/SwitchboardListener/SblCore/Private/SwitchboardListener.cpp
--- a/Engine/Source/Programs/SwitchboardListener/SblCore/Private/SwitchboardListener.cpp
+++ b/Engine/Source/Programs/SwitchboardListener/SblCore/Private/SwitchboardListener.cpp
@@ -326,6 +326,10 @@ FSwitchboardListener::FSwitchboardListener(const FSwitchboardCommandLineOptions&
 	}
 #endif // PLATFORM_WINDOWS
 
+	const int32 NumLogicalProcessors = FPlatformMisc::NumberOfCoresIncludingHyperthreads();
+	const int32 NumCores = FPlatformMisc::NumberOfCores();
+	bProcessorSMT = NumLogicalProcessors > NumCores;
+
 	const FIPv4Address DefaultIp = FIPv4Address(0, 0, 0, 0);
 	if (!Options.Address.IsSet())
 	{
@@ -1767,7 +1771,7 @@ QUIC_STATUS FSwitchboardListener::QuicConnectionCallback(HQUIC QuicConn, QUIC_CO
 
 				FPlatformMisc::GetOSVersions(StatePacket.OsVersionLabel, StatePacket.OsVersionLabelSub);
 				StatePacket.OsVersionNumber = FPlatformMisc::GetOSVersion();
-
+				StatePacket.bProcessorSMT = bProcessorSMT;
 				StatePacket.TotalPhysicalMemory = FPlatformMemory::GetConstants().TotalPhysical;
 				StatePacket.PlatformBinaryDirectory = FPlatformProcess::GetBinariesSubdirectory();
 
diff --git a/Engine/Source/Programs/SwitchboardListener/SblCore/Private/SwitchboardPacket.h b/Engine/Source/Programs/SwitchboardListener/SblCore/Private/SwitchboardPacket.h
--- a/Engine/Source/Programs/SwitchboardListener/SblCore/Private/SwitchboardPacket.h
+++ b/Engine/Source/Programs/SwitchboardListener/SblCore/Private/SwitchboardPacket.h
@@ -64,21 +64,31 @@ struct FSwitchboardStatePacket : public FSwitchboardPacket
 		Version = (Major << 16) | (Minor << 8) | (Patch);
 	}
 
+	/** All active child processes launched via this Listener. */
 	UPROPERTY()
 	TArray<FSwitchboardStateRunningProcess> RunningProcesses;
 
+	/** Reports the SwitchboardListener version back to the client. */
 	UPROPERTY()
 	uint32 Version;
 
+	/** Whether simultaneous multi-threading (aka Hyper-Threading) is enabled. */
+	UPROPERTY()
+	bool bProcessorSMT;
+
+	/** Operating system version "label". See FPlatformMisc::GetOSVersions. */
 	UPROPERTY()
 	FString OsVersionLabel;
 
+	/** Operating system version "sub-label". See FPlatformMisc::GetOSVersions. */
 	UPROPERTY()
 	FString OsVersionLabelSub;
 
+	/** Operating system version, returned by FPlatformMisc::GetOSVersion. */
 	UPROPERTY()
 	FString OsVersionNumber;
 
+	/** Total physical memory (used + available) in bytes. */
 	UPROPERTY()
 	uint64 TotalPhysicalMemory;
 
