Switchboard: Add context menu option to exclude individual devices from Unreal Insights tracing.

diff --git a/Engine/Plugins/VirtualProduction/Switchboard/Source/Switchboard/switchboard/devices/ndisplay/plugin_ndisplay.py b/Engine/Plugins/VirtualProduction/Switchboard/Source/Switchboard/switchboard/devices/ndisplay/plugin_ndisplay.py
--- a/Engine/Plugins/VirtualProduction/Switchboard/Source/Switchboard/switchboard/devices/ndisplay/plugin_ndisplay.py
+++ b/Engine/Plugins/VirtualProduction/Switchboard/Source/Switchboard/switchboard/devices/ndisplay/plugin_ndisplay.py
@@ -1031,8 +1031,7 @@ class DevicenDisplay(DeviceUnreal):
             dp_cvars.append('Slate.bAllowNotifications=0')
 
         # Insights traces parameters
-        if CONFIG.INSIGHTS_TRACE_ENABLE.get_value():
-
+        if CONFIG.INSIGHTS_TRACE_ENABLE.get_value() and not self.exclude_from_insights.get_value():
             LOGGER.warning(f"Unreal Insight Tracing is enabled for '{self.name}'. This may affect Unreal Engine performance.")
 
             remote_utrace_path = self.get_utrace_filepath()
diff --git a/Engine/Plugins/VirtualProduction/Switchboard/Source/Switchboard/switchboard/devices/unreal/plugin_unreal.py b/Engine/Plugins/VirtualProduction/Switchboard/Source/Switchboard/switchboard/devices/unreal/plugin_unreal.py
--- a/Engine/Plugins/VirtualProduction/Switchboard/Source/Switchboard/switchboard/devices/unreal/plugin_unreal.py
+++ b/Engine/Plugins/VirtualProduction/Switchboard/Source/Switchboard/switchboard/devices/unreal/plugin_unreal.py
@@ -923,6 +923,13 @@ class DeviceUnreal(Device):
             tool_tip="Whether to exclude this device from builds"
         )
 
+        self.exclude_from_insights = BoolSetting(
+            attr_name="exclude_from_insights",
+            nice_name="Exclude from Insights trace",
+            value=kwargs.get("exclude_from_insights", False),
+            tool_tip="Whether to exclude device from Unreal Insights traces"
+        )
+
         self.setting_address.signal_setting_changed.connect(
             self.on_setting_address_changed)
         DeviceUnreal.csettings['port'].signal_setting_changed.connect(
@@ -1039,6 +1046,11 @@ class DeviceUnreal(Device):
         )
         self.on_setting_exclude_from_build_changed(self.exclude_from_build.get_value())
 
+        self.exclude_from_insights.signal_setting_changed.connect(
+            lambda _, new_value: self.on_setting_exclude_from_insights_changed(new_value)
+        )
+        self.on_setting_exclude_from_insights_changed(self.exclude_from_insights.get_value())
+
     def should_allow_exit(self, close_req_id: int) -> bool:
         # Delegate to a class method which surveys all active devices.
         return DeviceUnreal._should_allow_exit(close_req_id)
@@ -1153,6 +1165,7 @@ class DeviceUnreal(Device):
             self.last_log_path,
             self.last_trace_path,
             self.exclude_from_build,
+            self.exclude_from_insights,
             self.last_sync_filter_hash,
         ]
 
@@ -1243,6 +1256,8 @@ class DeviceUnreal(Device):
 
         device_widget.signal_exclude_from_build_toggled.connect(self.on_toggle_exclude_from_build)
 
+        device_widget.signal_exclude_from_insights_toggled.connect(self.on_toggle_exclude_from_insights)
+
         # hook to open last log signal from widget
         device_widget.signal_open_last_log.connect(self.on_open_last_log)
 
@@ -1277,6 +1292,9 @@ class DeviceUnreal(Device):
     def on_setting_exclude_from_build_changed(self, exclude_from_build):
         self.widget.update_exclude_from_build(exclude_from_build, not self.is_disconnected)
 
+    def on_setting_exclude_from_insights_changed(self, exclude_from_insights):
+        self.widget.exclude_from_insights = exclude_from_insights
+
     @property
     def device_osc_port(self) -> int:
         return DeviceUnreal.csettings['osc_port'].get_value(self.name)
@@ -1653,7 +1671,7 @@ class DeviceUnreal(Device):
 
                 puuid_dependency = self._build_mu_server(
                     puuid_dependency=puuid_dependency)
- 
+
                 puuid_dependency = self._build_mu_slate_server(
                     puuid_dependency=puuid_dependency)
 
@@ -1924,7 +1942,7 @@ class DeviceUnreal(Device):
                               f'-CONCERTSESSION="{SETTINGS.MUSERVER_SESSION_NAME}" '
                               f'-CONCERTDISPLAYNAME="{self.name}"')
 
-        if CONFIG.INSIGHTS_TRACE_ENABLE.get_value():
+        if CONFIG.INSIGHTS_TRACE_ENABLE.get_value() and not self.exclude_from_insights.get_value():
             LOGGER.warning(f"Unreal Insight Tracing is enabled for '{self.name}'. This may affect Unreal Engine performance.")
             remote_utrace_path = self.get_utrace_filepath()
             command_line_args += ' -statnamedevents' if CONFIG.INSIGHTS_STAT_EVENTS.get_value() else ''
@@ -2809,7 +2827,13 @@ class DeviceUnreal(Device):
         self.update_settings_menu_state()
 
     def on_toggle_exclude_from_build(self):
-        self.exclude_from_build.update_value(not self.exclude_from_build.get_value())
+        self.exclude_from_build.update_value(
+            not self.exclude_from_build.get_value())
+        CONFIG.save()
+
+    def on_toggle_exclude_from_insights(self):
+        self.exclude_from_insights.update_value(
+            not self.exclude_from_insights.get_value())
         CONFIG.save()
 
     def on_open_last_log(self):
@@ -3051,10 +3075,12 @@ def parse_unreal_tag_file(file_content):
 
 BASE_ENGINE_CL_TOOLTIP = "Current Engine Changelist"
 BASE_PROJECT_CL_TOOLTIP = "Current Project Changelist"
-    
+
+
 class DeviceWidgetUnreal(DeviceWidget):
-    
+
     signal_exclude_from_build_toggled = QtCore.Signal()
+    signal_exclude_from_insights_toggled = QtCore.Signal()
     signal_open_last_log = QtCore.Signal(object)
     signal_open_last_trace = QtCore.Signal(object)
     signal_copy_last_launch_command = QtCore.Signal(object)
@@ -3067,6 +3093,7 @@ class DeviceWidgetUnreal(DeviceWidget):
         self._needs_rebuild = False
         self._exclude_from_build = False
         self._desired_build_button_tooltip = "Build changelist"
+        self.exclude_from_insights = False
 
         super().__init__(name, device_hash, address, icons, parent=parent)
 
@@ -3081,7 +3108,7 @@ class DeviceWidgetUnreal(DeviceWidget):
         self._exclude_from_build = exclude_from_build
         if not update_ui:
             return
-        
+
         if exclude_from_build:
             self.engine_changelist_label.hide()
             self._update_build_button_tooltip()
@@ -3292,7 +3319,7 @@ class DeviceWidgetUnreal(DeviceWidget):
         self.project_changelist_label.show()
         self.sync_button.show()
         self.build_button.show()
-    
+
         is_synched = required_cl is None or required_cl == current_device_cl
         self._set_project_changelist_is_synched(is_synched)
 
@@ -3312,14 +3339,14 @@ class DeviceWidgetUnreal(DeviceWidget):
     def update_engine_changelist(self, required_cl: str, synched_cl: str, built__cl: str):
         if not CONFIG.ENGINE_SYNC_METHOD.get_value() == EngineSyncMethod.Build_Engine.value:
             return
-        
+
         self.engine_changelist_label.setText(f'E: {synched_cl}')
         if not self.exclude_from_build:
             self.engine_changelist_label.show()
 
         self.sync_button.show()
         self.build_button.show()
-        
+
         is_synched = required_cl is None or required_cl == synched_cl
         self._set_engine_changelist_is_synched(is_synched)
         self.update_build_info(synched_cl=synched_cl, built_cl=built__cl)
@@ -3327,21 +3354,21 @@ class DeviceWidgetUnreal(DeviceWidget):
     def _set_project_changelist_is_synched(self, is_synched: bool):
         self._is_project_synched = is_synched
         sb_widgets.set_qt_property(self.project_changelist_label, 'not_synched', not is_synched)
-        
+
         self._update_cl_widget_tooltip(self.project_changelist_label, BASE_PROJECT_CL_TOOLTIP, self._is_project_synched)
         self._update_sync_button()
 
     def _set_engine_changelist_is_synched(self, is_synched: bool):
         self._is_engine_synched = is_synched
         self._update_engine_cl_label()
-        
+
         self._update_cl_widget_tooltip(self.engine_changelist_label, BASE_ENGINE_CL_TOOLTIP, self._is_engine_synched)
         self._update_sync_button()
-        
+
     def _update_sync_button(self):
         needs_resync = not self._is_project_synched or not self._is_engine_synched
         sb_widgets.set_qt_property(self.sync_button, 'not_synched', needs_resync)
-            
+
     def update_build_info(self, synched_cl: str, built_cl: str):
         if built_cl is not None and synched_cl is not None and CONFIG.ENGINE_SYNC_METHOD.get_value() == EngineSyncMethod.Build_Engine.value:
             try:
@@ -3365,7 +3392,7 @@ class DeviceWidgetUnreal(DeviceWidget):
         self._desired_build_button_tooltip = desired_tooltip
         self._update_build_button_tooltip()
         self._refresh_build_info_ui()
-        
+
     def _refresh_build_info_ui(self):
         should_update_ui = not self.exclude_from_build
         if should_update_ui:
@@ -3373,13 +3400,13 @@ class DeviceWidgetUnreal(DeviceWidget):
             self._update_build_button_tooltip()
             self._update_engine_cl_label()
             self._update_cl_widget_tooltip(self.engine_changelist_label, BASE_ENGINE_CL_TOOLTIP, self._is_engine_synched)
-            
+
     def _update_build_button_tooltip(self):
         if self.exclude_from_build:
             self.build_button.setToolTip("Excluded from build (see device settings)")
         else:
             self.build_button.setToolTip(self._desired_build_button_tooltip)
-            
+
     def _update_engine_cl_label(self):
         if self.exclude_from_build:
             sb_widgets.set_qt_property(self.engine_changelist_label, 'not_synched', False)
@@ -3399,7 +3426,7 @@ class DeviceWidgetUnreal(DeviceWidget):
 
         if not is_synched:
             tooltip += "\nNot synched to selected CL"
-            
+
         if self._needs_rebuild:
             tooltip += "\nSynched CL not built"
 
@@ -3424,12 +3451,16 @@ class DeviceWidgetUnreal(DeviceWidget):
             self._disconnect()
 
     def populate_context_menu(self, cmenu: QtWidgets.QMenu):
-        ''' Called to populate the given context menu with any desired actions'''
-        
+        ''' Called to populate the given context menu with any desired actions '''
+
         cmenu.addAction(
             "Include in build" if self.exclude_from_build else "Exclude from build",
             lambda: self.signal_exclude_from_build_toggled.emit()
         )
+        cmenu.addAction(
+            "Include in Insights traces" if self.exclude_from_insights else "Exclude from Insights traces",
+            lambda: self.signal_exclude_from_insights_toggled.emit()
+        )
         cmenu.addAction("Open fetched log", lambda: self.signal_open_last_log.emit(self))
         cmenu.addAction("Open fetched trace", lambda: self.signal_open_last_trace.emit(self))
         cmenu.addAction("Copy last launch command", lambda: self.signal_copy_last_launch_command.emit(self))
@@ -3437,6 +3468,6 @@ class DeviceWidgetUnreal(DeviceWidget):
         # Only create DDC submenu if node is connected
         if self.connect_button.isChecked():
             fill_ddc_menu = cmenu.addMenu("Fill DDC (Prepare Shaders)")
-        
+
             current_level_action = fill_ddc_menu.addAction("Current Level", lambda: self.signal_device_widget_fill_ddc.emit(self, True))
             all_levels_action = fill_ddc_menu.addAction("All Levels", lambda: self.signal_device_widget_fill_ddc.emit(self, False))
