Switchboard: Automatically trim .uproject filename from P4 project path if accidentally specified.

diff --git a/Engine/Plugins/VirtualProduction/Switchboard/Source/Switchboard/switchboard/config.py b/Engine/Plugins/VirtualProduction/Switchboard/Source/Switchboard/switchboard/config.py
--- a/Engine/Plugins/VirtualProduction/Switchboard/Source/Switchboard/switchboard/config.py
+++ b/Engine/Plugins/VirtualProduction/Switchboard/Source/Switchboard/switchboard/config.py
@@ -640,12 +640,16 @@ class StringSetting(Setting):
         self.set_widget(
             widget=line_edit, override_device_name=override_device_name)
 
-        line_edit.editingFinished.connect(
-            lambda line_edit=line_edit,
-            override_device_name=override_device_name:
-                self._on_widget_value_changed(
-                    line_edit.text().strip(),
-                    override_device_name=override_device_name))
+        def on_editing_finished():
+            self._on_widget_value_changed(
+                line_edit.text().strip(),
+                override_device_name=override_device_name)
+
+            # In case the value the user entered is filtered to be the same as
+            # the prior value, control should reflect that it was "rejected."
+            line_edit.setText(self.get_value(override_device_name))
+
+        line_edit.editingFinished.connect(on_editing_finished)
 
         return line_edit
 
@@ -776,6 +780,26 @@ class PerforcePathSetting(StringSetting):
     represents a Perforce depot path.
     '''
 
+    def __init__(
+        self,
+        attr_name: str,
+        nice_name: str,
+        value: str,
+        placeholder_text: str = '',
+        tool_tip: Optional[str] = None,
+        show_ui: bool = True,
+        allow_reset: bool = True,
+        migrate_data: Optional[Callable[[Any], None]] = None,
+        is_read_only: bool = False,
+    ):
+        # Trim matching file paths to the parent directory (e.g. ['.uproject'])
+        self.truncate_files_with_extensions: list[str] = []
+
+        super().__init__(
+            attr_name, nice_name, value, placeholder_text=placeholder_text,
+            tool_tip=tool_tip, show_ui=show_ui, allow_reset=allow_reset,
+            migrate_data=migrate_data, is_read_only=is_read_only)
+
     def _filter_value(self, value: Optional[str]) -> str:
         '''
         Clean the p4 path value by removing whitespace and trailing '/'.
@@ -783,7 +807,16 @@ class PerforcePathSetting(StringSetting):
         if not value:
             return ''
 
-        return value.strip().rstrip('/')
+        value = value.strip()
+
+        # Fix up common case where user incorrectly includes a file name.
+        value_lower = value.lower()
+        for ext in self.truncate_files_with_extensions:
+            if value_lower.endswith(ext.lower()):
+                value = value[:value.rfind('/')]
+                break
+
+        return value.rstrip('/')
 
 
 class OptionSetting(Setting):
@@ -2328,8 +2361,10 @@ class Config(object):
 
         self.P4_ENABLED = self.source_control_settings["p4_enabled"]
         self.SOURCE_CONTROL_WORKSPACE = self.source_control_settings["source_control_workspace"]
-        self.P4_PROJECT_PATH = self.source_control_settings["p4_sync_path"]
-        self.P4_ENGINE_PATH = self.source_control_settings["p4_engine_path"]
+        self.P4_PROJECT_PATH: PerforcePathSetting = self.source_control_settings["p4_sync_path"]
+        self.P4_ENGINE_PATH: PerforcePathSetting = self.source_control_settings["p4_engine_path"]
+
+        self.P4_PROJECT_PATH.truncate_files_with_extensions.append('.uproject')
 
     def init_unreal_insights(self, data={}):
         self.unreal_insight_settings = {
