Switchboard: Only send OSC slate change message when text edit loses focus, instead of every keystroke.

diff --git a/Engine/Plugins/VirtualProduction/Switchboard/Source/Switchboard/switchboard/switchboard_dialog.py b/Engine/Plugins/VirtualProduction/Switchboard/Source/Switchboard/switchboard/switchboard_dialog.py
--- a/Engine/Plugins/VirtualProduction/Switchboard/Source/Switchboard/switchboard/switchboard_dialog.py
+++ b/Engine/Plugins/VirtualProduction/Switchboard/Source/Switchboard/switchboard/switchboard_dialog.py
@@ -404,7 +404,8 @@ class SwitchboardDialog(QtCore.QObject):
 
         # Connect UI to methods
         self.window.multiuser_session_lineEdit.textChanged.connect(self.on_multiuser_session_lineEdit_textChanged)
-        self.window.slate_line_edit.textChanged.connect(self._set_slate)
+        self.window.slate_line_edit.editingFinished.connect(
+            lambda: self._set_slate(self.window.slate_line_edit.text()))
         self.window.take_spin_box.valueChanged.connect(self._set_take)
         self.window.sequence_line_edit.textChanged.connect(self._set_sequence)
         self.level_combo_box.currentIndexChanged.connect(self._on_selected_level_changed)
