Switchboard: Explicitly hide child processes that were previously implicitly hidden (console subsystem -> console subsystem).

diff --git a/Engine/Plugins/VirtualProduction/Switchboard/Source/Switchboard/switchboard/devices/unreal/plugin_unreal.py b/Engine/Plugins/VirtualProduction/Switchboard/Source/Switchboard/switchboard/devices/unreal/plugin_unreal.py
--- a/Engine/Plugins/VirtualProduction/Switchboard/Source/Switchboard/switchboard/devices/unreal/plugin_unreal.py
+++ b/Engine/Plugins/VirtualProduction/Switchboard/Source/Switchboard/switchboard/devices/unreal/plugin_unreal.py
@@ -1373,6 +1373,7 @@ class DeviceUnreal(Device):
             caller=self.name,
             working_dir=working_dir,
             update_clients_with_stdout=False,
+            hide=True,
         )
 
         self.program_start_queue.add(
@@ -1419,6 +1420,7 @@ class DeviceUnreal(Device):
             caller=self.name,
             working_dir=working_dir,
             update_clients_with_stdout=False,
+            hide=True,
         )
 
         self.program_start_queue.add(
@@ -1594,6 +1596,7 @@ class DeviceUnreal(Device):
             working_dir=os.path.dirname(
                 CONFIG.UPROJECT_PATH.get_value(self.name)),
             update_clients_with_stdout=True,
+            hide=True,
         )
 
         self.program_start_queue.add(
@@ -1746,6 +1749,7 @@ class DeviceUnreal(Device):
             prog_name=program_name,
             caller=self.name,
             update_clients_with_stdout=True,
+            hide=True,
         )
 
         def launch_fn():
@@ -2717,7 +2721,8 @@ class DeviceUnreal(Device):
             prog_path=str(rsync_path),
             prog_args=rsync_args,
             prog_name=program_name,
-            caller=self.name
+            caller=self.name,
+            hide=True,
         )
 
         self.program_start_queue.add(
diff --git a/Engine/Plugins/VirtualProduction/Switchboard/Source/Switchboard/switchboard/message_protocol.py b/Engine/Plugins/VirtualProduction/Switchboard/Source/Switchboard/switchboard/message_protocol.py
--- a/Engine/Plugins/VirtualProduction/Switchboard/Source/Switchboard/switchboard/message_protocol.py
+++ b/Engine/Plugins/VirtualProduction/Switchboard/Source/Switchboard/switchboard/message_protocol.py
@@ -25,6 +25,7 @@ def create_start_process_message(
     update_clients_with_stdout: bool = False,
     priority_modifier: int = 0,
     lock_gpu_clock: bool = False,
+    hide: bool = False,
 ):
     cmd_id = uuid.uuid4()
     start_cmd = {
@@ -37,7 +38,8 @@ def create_start_process_message(
         'working_dir': working_dir,
         'bUpdateClientsWithStdout': update_clients_with_stdout,
         'priority_modifier': priority_modifier,
-        'bLockGpuClock' : lock_gpu_clock,
+        'bLockGpuClock': lock_gpu_clock,
+        'bHide': hide,
     }
 
     message = json.dumps(start_cmd).encode() + b'\x00'
diff --git a/Engine/Source/Programs/SwitchboardListener/SblCore/Internal/SwitchboardListenerVersion.h b/Engine/Source/Programs/SwitchboardListener/SblCore/Internal/SwitchboardListenerVersion.h
--- a/Engine/Source/Programs/SwitchboardListener/SblCore/Internal/SwitchboardListenerVersion.h
+++ b/Engine/Source/Programs/SwitchboardListener/SblCore/Internal/SwitchboardListenerVersion.h
@@ -4,4 +4,4 @@
 
 #define SBLISTENER_VERSION_MAJOR 3
 #define SBLISTENER_VERSION_MINOR 0
-#define SBLISTENER_VERSION_PATCH 3
+#define SBLISTENER_VERSION_PATCH 4
diff --git a/Engine/Source/Programs/SwitchboardListener/SblCore/Private/SwitchboardListener.cpp b/Engine/Source/Programs/SwitchboardListener/SblCore/Private/SwitchboardListener.cpp
--- a/Engine/Source/Programs/SwitchboardListener/SblCore/Private/SwitchboardListener.cpp
+++ b/Engine/Source/Programs/SwitchboardListener/SblCore/Private/SwitchboardListener.cpp
@@ -1026,10 +1026,10 @@ bool FSwitchboardListener::Task_StartProcess(const FSwitchboardStartTask& InRunT
 	{
 		return false;
 	}
-	
+
 	const bool bLaunchDetached = false;
-	const bool bLaunchHidden = false;
-	const bool bLaunchReallyHidden = false;
+	const bool bLaunchHidden = InRunTask.bHide;
+	const bool bLaunchReallyHidden = InRunTask.bHide;
 	const int32 PriorityModifier = InRunTask.PriorityModifier;
 	const TCHAR* WorkingDirectory = InRunTask.WorkingDir.IsEmpty() ? nullptr : *InRunTask.WorkingDir;
 
diff --git a/Engine/Source/Programs/SwitchboardListener/SblCore/Private/SwitchboardProtocol.cpp b/Engine/Source/Programs/SwitchboardListener/SblCore/Private/SwitchboardProtocol.cpp
--- a/Engine/Source/Programs/SwitchboardListener/SblCore/Private/SwitchboardProtocol.cpp
+++ b/Engine/Source/Programs/SwitchboardListener/SblCore/Private/SwitchboardProtocol.cpp
@@ -254,6 +254,11 @@ FCreateTaskResult CreateTaskFromCommand(const FString& InCommand, const FIPv4End
 			PriorityModifierField->TryGetNumber(Task->PriorityModifier);
 		}
 
+		if (TSharedPtr<FJsonValue> HideField = JsonData->TryGetField(TEXT("bHide")))
+		{
+			HideField->TryGetBool(Task->bHide);
+		}
+
 		Result.Status = ECreateTaskStatus::Success;
 		Result.Task = MoveTemp(Task);
 		return Result;
diff --git a/Engine/Source/Programs/SwitchboardListener/SblCore/Private/SwitchboardTasks.h b/Engine/Source/Programs/SwitchboardListener/SblCore/Private/SwitchboardTasks.h
--- a/Engine/Source/Programs/SwitchboardListener/SblCore/Private/SwitchboardTasks.h
+++ b/Engine/Source/Programs/SwitchboardListener/SblCore/Private/SwitchboardTasks.h
@@ -131,6 +131,7 @@ struct FSwitchboardStartTask : public FSwitchboardTask
 	bool bUpdateClientsWithStdout = false;
 	bool bLockGpuClock = false;
 	int32 PriorityModifier = 0;
+	bool bHide = false;
 
 	//~ Begin FSwitchboardTask interface
 	static constexpr const TCHAR* CommandName = TEXT("start");
