Switchboard: Authentication fixes, bump Listener to 3.0.5.

- Don't send FSwitchboardStatePacket to clients until after they authenticate successfully.
- Made some calls to Future.set_exception thread-safe.
- Exposed authenticated status to outside code, so that nDisplay Monitor doesn't try to poll sync status until after we auth.
- Moved some cleanup code to _loop_threadproc, which could otherwise result in a half-connected client. Additionally handle exceptions besides ConnectionError.
- Don't run asyncio loop in debug mode to cut down on log spam.

#rnx

diff --git a/Engine/Plugins/VirtualProduction/Switchboard/Source/Switchboard/switchboard/devices/ndisplay/ndisplay_monitor.py b/Engine/Plugins/VirtualProduction/Switchboard/Source/Switchboard/switchboard/devices/ndisplay/ndisplay_monitor.py
--- a/Engine/Plugins/VirtualProduction/Switchboard/Source/Switchboard/switchboard/devices/ndisplay/ndisplay_monitor.py
+++ b/Engine/Plugins/VirtualProduction/Switchboard/Source/Switchboard/switchboard/devices/ndisplay/ndisplay_monitor.py
@@ -269,8 +269,9 @@ class nDisplayMonitor(QAbstractTableModel):
             # detect stale devices
             self.handle_stale_device(devicedata, deviceIdx)
 
-            # no point in continuing of not connected to listener
-            if not device.unreal_client.is_connected:
+            # no point in continuing if not connected to listener
+            if not (device.unreal_client.is_connected
+                    and device.unreal_client.is_authenticated):
                 continue
 
             # create message
diff --git a/Engine/Plugins/VirtualProduction/Switchboard/Source/Switchboard/switchboard/listener_client.py b/Engine/Plugins/VirtualProduction/Switchboard/Source/Switchboard/switchboard/listener_client.py
--- a/Engine/Plugins/VirtualProduction/Switchboard/Source/Switchboard/switchboard/listener_client.py
+++ b/Engine/Plugins/VirtualProduction/Switchboard/Source/Switchboard/switchboard/listener_client.py
@@ -61,11 +61,14 @@ class SwitchboardClientProtocol(QuicConnectionProtocol):
             cred: Optional[CredentialStore.Credential],
             unsaved: bool
         ):
+            def auth_exception():
+                auth_response.set_exception(RuntimeError)
+                self.close()
+
             if not cred:
                 self._quic._logger.error(
                     'No credential for %s', host)
-                auth_response.set_exception(RuntimeError)
-                self.close()
+                self._loop.call_soon_threadsafe(auth_exception)
                 return
 
             nonlocal closure_credential, closure_cred_unsaved
@@ -81,8 +84,7 @@ class SwitchboardClientProtocol(QuicConnectionProtocol):
                     '%s (expected)\n'
                     '%s (actual)',
                     host, cred.username, fingerprint_str)
-                auth_response.set_exception(RuntimeError)
-                self.close()
+                self._loop.call_soon_threadsafe(auth_exception)
                 return
 
             if self.client:
@@ -92,7 +94,7 @@ class SwitchboardClientProtocol(QuicConnectionProtocol):
                 self.client.get_message_response(msgid, msg,
                                                  future=auth_response)
             else:
-                auth_response.set_exception(RuntimeError)
+                self._loop.call_soon_threadsafe(auth_exception)
 
         CREDENTIAL_STORE.get(
             key=host,
@@ -145,6 +147,7 @@ class ListenerClient:
     EXTRA_HITCH_TOLERANCE_SEC = 0.4
     HITCH_THRESHOLD_SEC = (KEEPALIVE_INTERVAL_SEC + SELECT_TIMEOUT_SEC
                            + EXTRA_HITCH_TOLERANCE_SEC)
+    IDLE_TIMEOUT_SEC = 5.0
 
     def __init__(self, address, port, buffer_size=1024):
         self.address = address
@@ -158,6 +161,7 @@ class ListenerClient:
         self.loop: Optional[asyncio.AbstractEventLoop] = None
         self.loop_thread: Optional[threading.Thread] = None
         self.close_connection = False
+        self.is_authenticated = False
 
         # TODO: Consider converting these delegates to Signals and sending dict
 
@@ -216,6 +220,7 @@ class ListenerClient:
             return False
 
         self.close_connection = False
+        self.is_authenticated = False
         self.last_activity = datetime.now()
 
         self.listener_qt_handler.listener_connecting.emit(self)
@@ -226,7 +231,7 @@ class ListenerClient:
         return True
 
     def disconnect(self, unexpected=False, exception=None):
-        if self.protocol:
+        if self.protocol and self.loop:
             _, msg = message_protocol.create_disconnect_message()
             self.send_message(msg)
             self.loop.call_soon_threadsafe(self.protocol.close)
@@ -248,12 +253,17 @@ class ListenerClient:
 
     def _loop_threadproc(self, **kwargs):
         try:
-            _ = asyncio.run(self._async_main(**kwargs), debug=True)
-        except ConnectionError as exc:
+            _ = asyncio.run(self._async_main(**kwargs))
+        except Exception as exc:
+            self.listener_qt_handler.listener_connection_failed.emit(self)
             if self.disconnect_delegate:
                 self.disconnect_delegate(unexpected=True, exception=exc)
 
+        self.reader = None
+        self.writer = None
+        self.protocol = None
         self.loop = None
+        self.is_authenticated = False
 
     async def _async_main(self, timeout: Optional[float]):
         LOGGER.info(f"Connecting to {self.address}:{self.port}")
@@ -264,7 +274,7 @@ class ListenerClient:
             alpn_protocols=[SWITCHBOARD_ALPN],
             is_client=True)
 
-        configuration.idle_timeout = timeout or 10.0
+        configuration.idle_timeout = timeout or self.IDLE_TIMEOUT_SEC
 
         # We deal exclusively with self-signed certificates in practice,
         # and depend trust-on-first-use and validating the fingerprint.
@@ -283,18 +293,14 @@ class ListenerClient:
             self.reader._limit = 2 ** 20  # 1 MiB max message (default 64 KiB)
             message_recv_task = asyncio.create_task(self._message_recv_loop())
 
-            is_authenticated = await self.protocol.authenticate()
-            if is_authenticated:
+            self.is_authenticated = await self.protocol.authenticate()
+            if self.is_authenticated:
                 self.listener_qt_handler.listener_connected.emit(self)
                 await asyncio.wait_for(message_recv_task, timeout=None)
             else:
                 LOGGER.error('Authentication failed!')
                 self.listener_qt_handler.listener_connection_failed.emit(self)
 
-        self.reader = None
-        self.writer = None
-        self.protocol = None
-
         if self.disconnect_delegate:
             self.disconnect_delegate(
                 unexpected=(not self.close_connection),
@@ -427,7 +433,7 @@ class ListenerClient:
         message_id: uuid.UUID,
         message_bytes: bytes,
         *,
-        max_wait=5.0,
+        max_wait=IDLE_TIMEOUT_SEC,
         future: Optional[asyncio.Future[dict[str, str]]] = None
     ) -> asyncio.Future[dict[str, str]]:
         if self.loop and self.protocol and self.writer:
diff --git a/Engine/Source/Programs/SwitchboardListener/SblCore/Internal/SwitchboardListener.h b/Engine/Source/Programs/SwitchboardListener/SblCore/Internal/SwitchboardListener.h
--- a/Engine/Source/Programs/SwitchboardListener/SblCore/Internal/SwitchboardListener.h
+++ b/Engine/Source/Programs/SwitchboardListener/SblCore/Internal/SwitchboardListener.h
@@ -157,7 +157,6 @@ public:
 	FSimpleMulticastDelegate& OnTick() { return OnTickDelegate; }
 
 private:
-
 	bool ParseIncomingMessage(const FString& InMessage, const FIPv4Endpoint& InEndpoint, const FConnectionRef& Connection);
 
 	bool RunScheduledTask(const FSwitchboardTask& InTask);
@@ -188,6 +187,8 @@ private:
 
 	void RollbackRedeploy();
 
+	void FillStatePacket(struct FSwitchboardStatePacket& OutStatePacket);
+
 	_Function_class_(QUIC_LISTENER_CALLBACK)
 	static QUIC_STATUS QUIC_API QuicListenerThunk(HQUIC Listener, void* Context, QUIC_LISTENER_EVENT* Event);
 	static QUIC_STATUS QUIC_API QuicConnectionThunk(HQUIC Connection, void* Context, QUIC_CONNECTION_EVENT* Event);
diff --git a/Engine/Source/Programs/SwitchboardListener/SblCore/Internal/SwitchboardListenerVersion.h b/Engine/Source/Programs/SwitchboardListener/SblCore/Internal/SwitchboardListenerVersion.h
--- a/Engine/Source/Programs/SwitchboardListener/SblCore/Internal/SwitchboardListenerVersion.h
+++ b/Engine/Source/Programs/SwitchboardListener/SblCore/Internal/SwitchboardListenerVersion.h
@@ -4,4 +4,4 @@
 
 #define SBLISTENER_VERSION_MAJOR 3
 #define SBLISTENER_VERSION_MINOR 0
-#define SBLISTENER_VERSION_PATCH 4
+#define SBLISTENER_VERSION_PATCH 5
diff --git a/Engine/Source/Programs/SwitchboardListener/SblCore/Private/SwitchboardListener.cpp b/Engine/Source/Programs/SwitchboardListener/SblCore/Private/SwitchboardListener.cpp
--- a/Engine/Source/Programs/SwitchboardListener/SblCore/Private/SwitchboardListener.cpp
+++ b/Engine/Source/Programs/SwitchboardListener/SblCore/Private/SwitchboardListener.cpp
@@ -998,6 +998,14 @@ bool FSwitchboardListener::Task_Authenticate(const FSwitchboardAuthenticateTask&
 		Connection->bAuthenticated = true;
 		// TODO: Issue JWT?
 		SendAuthResponse();
+
+		// Send current state upon authentication
+		{
+			FSwitchboardStatePacket StatePacket;
+			FillStatePacket(StatePacket);
+			SendMessage(CreateMessage(StatePacket), InAuthTask.Recipient);
+		}
+
 		return true;
 	}
 
@@ -1640,6 +1648,32 @@ bool FSwitchboardListener::Task_SendFileToClient(const FSwitchboardSendFileToCli
 }
 
 
+void FSwitchboardListener::FillStatePacket(FSwitchboardStatePacket& OutStatePacket)
+{
+	FPlatformMisc::GetOSVersions(OutStatePacket.OsVersionLabel, OutStatePacket.OsVersionLabelSub);
+	OutStatePacket.OsVersionNumber = FPlatformMisc::GetOSVersion();
+	OutStatePacket.bProcessorSMT = bProcessorSMT;
+	OutStatePacket.TotalPhysicalMemory = FPlatformMemory::GetConstants().TotalPhysical;
+	OutStatePacket.PlatformBinaryDirectory = FPlatformProcess::GetBinariesSubdirectory();
+
+	OutStatePacket.RunningProcesses.Empty(RunningProcesses.Num());
+	for (const TSharedPtr<FRunningProcess>& RunningProcess : RunningProcesses)
+	{
+		check(RunningProcess.IsValid());
+
+		FSwitchboardStateRunningProcess StateRunningProcess;
+
+		StateRunningProcess.Uuid = RunningProcess->UUID.ToString();
+		StateRunningProcess.Name = RunningProcess->Name;
+		StateRunningProcess.Path = RunningProcess->Path;
+		StateRunningProcess.Caller = RunningProcess->Caller;
+		StateRunningProcess.Pid = RunningProcess->PID;
+
+		OutStatePacket.RunningProcesses.Add(MoveTemp(StateRunningProcess));
+	}
+}
+
+
 //static
 _Function_class_(QUIC_LISTENER_CALLBACK)
 QUIC_STATUS QUIC_API FSwitchboardListener::QuicListenerThunk(HQUIC Listener, void* Context, QUIC_LISTENER_EVENT* Event)
@@ -1765,34 +1799,6 @@ QUIC_STATUS FSwitchboardListener::QuicConnectionCallback(HQUIC QuicConn, QUIC_CO
 
 			QuicApi->SetCallbackHandler(Connection->QuicStream, reinterpret_cast<void*>(&QuicStreamThunk), this);
 
-			// Send current state upon connection
-			{
-				FSwitchboardStatePacket StatePacket;
-
-				FPlatformMisc::GetOSVersions(StatePacket.OsVersionLabel, StatePacket.OsVersionLabelSub);
-				StatePacket.OsVersionNumber = FPlatformMisc::GetOSVersion();
-				StatePacket.bProcessorSMT = bProcessorSMT;
-				StatePacket.TotalPhysicalMemory = FPlatformMemory::GetConstants().TotalPhysical;
-				StatePacket.PlatformBinaryDirectory = FPlatformProcess::GetBinariesSubdirectory();
-
-				for (const TSharedPtr<FRunningProcess>& RunningProcess : RunningProcesses)
-				{
-					check(RunningProcess.IsValid());
-
-					FSwitchboardStateRunningProcess StateRunningProcess;
-
-					StateRunningProcess.Uuid = RunningProcess->UUID.ToString();
-					StateRunningProcess.Name = RunningProcess->Name;
-					StateRunningProcess.Path = RunningProcess->Path;
-					StateRunningProcess.Caller = RunningProcess->Caller;
-					StateRunningProcess.Pid = RunningProcess->PID;
-
-					StatePacket.RunningProcesses.Add(MoveTemp(StateRunningProcess));
-				}
-
-				SendMessage(CreateMessage(StatePacket), RemoteEndpoint);
-			}
-
 			break;
 		}
 		case QUIC_CONNECTION_EVENT_RESUMED:
