Fix p4 workspace discovery not working when the path had a trailing directory separator.

diff --git a/Engine/Plugins/VirtualProduction/Switchboard/Source/Switchboard/switchboard/p4_utils.py b/Engine/Plugins/VirtualProduction/Switchboard/Source/Switchboard/switchboard/p4_utils.py
--- a/Engine/Plugins/VirtualProduction/Switchboard/Source/Switchboard/switchboard/p4_utils.py
+++ b/Engine/Plugins/VirtualProduction/Switchboard/Source/Switchboard/switchboard/p4_utils.py
@@ -126,6 +126,9 @@ def p4_from_localpath(localpath, workspaces, preferredClient):
         str,str: The workspace name and matching p4 path.
     '''
 
+    # Normalize path. In particular, a trailing slash may not be accepted by the 'p4 where' command.
+    localpath = pathlib.Path(localpath).as_posix()
+
     # Only take into account workspaces with the same give local path
     wss = [ws for ws in workspaces if workspaceInPath(ws, localpath)]
 
