Switchboard: Fix not selecting the map when creating a new config from Unreal.

The issue is that the list of maps being internally listed with the pattern /Game/MyFolder/MyMap.umap/MyMap, instead of just /Game/MyFolder/MyMap, so it never found the entry from the script and was falling back to the default map.

diff --git a/Engine/Plugins/VirtualProduction/Switchboard/Source/Switchboard/switchboard/config.py b/Engine/Plugins/VirtualProduction/Switchboard/Source/Switchboard/switchboard/config.py
--- a/Engine/Plugins/VirtualProduction/Switchboard/Source/Switchboard/switchboard/config.py
+++ b/Engine/Plugins/VirtualProduction/Switchboard/Source/Switchboard/switchboard/config.py
@@ -13,6 +13,7 @@ import threading
 import time
 from typing import Any, Callable, Optional, Tuple, Type, Union
 from enum import Enum
+from pathlib import Path
 
 from PySide2 import QtCore
 from PySide2 import QtGui
@@ -2864,15 +2865,14 @@ class Config(object):
                 if not fnmatch.fnmatch(umap.name, maps_filter):
                     continue
 
-                map_name, _ = os.path.splitext(umap.name)
-                file_path_to_map = os.path.join(umap.path, map_name)
+                mapgamepath = Path(self.resolve_content_path(
+                    umap.path,
+                    unreal_content_plugin=unreal_content_plugin))
 
-                content_path_to_map = self.resolve_content_path(
-                    file_path_to_map,
-                    unreal_content_plugin=unreal_content_plugin)
+                mapgamepath = (mapgamepath.parent / mapgamepath.stem).as_posix()
 
-                if content_path_to_map not in levels:
-                    levels.append(content_path_to_map)
+                if mapgamepath not in levels:
+                    levels.append(mapgamepath)
 
         levels.sort()
 
