[nDisplay] removed legacy, unnecesary code that prevents from launching dedicated server on the same host with Switchboard

diff --git a/Engine/Plugins/VirtualProduction/Switchboard/Source/Switchboard/switchboard/devices/ndisplay/plugin_ndisplay.py b/Engine/Plugins/VirtualProduction/Switchboard/Source/Switchboard/switchboard/devices/ndisplay/plugin_ndisplay.py
--- a/Engine/Plugins/VirtualProduction/Switchboard/Source/Switchboard/switchboard/devices/ndisplay/plugin_ndisplay.py
+++ b/Engine/Plugins/VirtualProduction/Switchboard/Source/Switchboard/switchboard/devices/ndisplay/plugin_ndisplay.py
@@ -756,9 +756,6 @@ class DevicenDisplay(DeviceUnreal):
         if multiplayer_mode_name == 'Dedicated server':
             dedicated_server_address = DevicenDisplay.csettings["dedicated_server_address"].get_value()
 
-            if dedicated_server_address == self.address or is_local_address(dedicated_server_address):
-                dedicated_server_address = '127.0.0.1'
-
             return get_client_args(
                 f'{dedicated_server_address}:{DevicenDisplay.csettings["dedicated_server_port"].get_value()}')
 
