Switchboard: Added missing vulkan sm5 and sm6 options.

diff --git a/Engine/Plugins/VirtualProduction/Switchboard/Source/Switchboard/switchboard/devices/ndisplay/plugin_ndisplay.py b/Engine/Plugins/VirtualProduction/Switchboard/Source/Switchboard/switchboard/devices/ndisplay/plugin_ndisplay.py
--- a/Engine/Plugins/VirtualProduction/Switchboard/Source/Switchboard/switchboard/devices/ndisplay/plugin_ndisplay.py
+++ b/Engine/Plugins/VirtualProduction/Switchboard/Source/Switchboard/switchboard/devices/ndisplay/plugin_ndisplay.py
@@ -324,10 +324,12 @@ class DevicenDisplay(DeviceUnreal):
                 "dx11",
                 "dx11 -sm5",
                 "dx11 -sm6",
-                "dx12", 
-                "dx12 -sm5", 
-                "dx12 -sm6", 
-                "vulkan"
+                "dx12",
+                "dx12 -sm5",
+                "dx12 -sm6",
+                "vulkan",
+                "vulkan -sm5",
+                "vulkan -sm6"
             ],
         ),
         'multiplayer_mode': OptionSetting(
