Switchboard: Fix exception when removing nDisplay devices.

Also includes a fix for an unrelated exception when updating UGS sync filters that could occur when only nDisplay devices have been added.

#rnx

diff --git a/Engine/Plugins/VirtualProduction/Switchboard/Source/Switchboard/switchboard/devices/unreal/plugin_unreal.py b/Engine/Plugins/VirtualProduction/Switchboard/Source/Switchboard/switchboard/devices/unreal/plugin_unreal.py
--- a/Engine/Plugins/VirtualProduction/Switchboard/Source/Switchboard/switchboard/devices/unreal/plugin_unreal.py
+++ b/Engine/Plugins/VirtualProduction/Switchboard/Source/Switchboard/switchboard/devices/unreal/plugin_unreal.py
@@ -785,7 +785,7 @@ class DeviceUnreal(Device):
             sync_filters.categories.items(), key=lambda item: item[1].name))
 
         # The UUIDs are stable, but the names/object identities change
-        include_setting = cls.csettings['included_sync_categories']
+        include_setting = DeviceUnreal.csettings['included_sync_categories']
         include_setting.possible_values.clear()
         for category in sync_filters.categories.values():
             option = SyncCategoryOption(category.id, category.name)
@@ -2667,6 +2667,7 @@ class DeviceUnreal(Device):
 
         return None
 
+    @QtCore.Slot()
     def reregister_rsync_client(self):
         DeviceUnreal.rsync_server.unregister_client(self)
         DeviceUnreal.rsync_server.register_client(
diff --git a/Engine/Plugins/VirtualProduction/Switchboard/Source/Switchboard/switchboard/settings_dialog.py b/Engine/Plugins/VirtualProduction/Switchboard/Source/Switchboard/switchboard/settings_dialog.py
--- a/Engine/Plugins/VirtualProduction/Switchboard/Source/Switchboard/switchboard/settings_dialog.py
+++ b/Engine/Plugins/VirtualProduction/Switchboard/Source/Switchboard/switchboard/settings_dialog.py
@@ -33,7 +33,7 @@ class SettingsDialog(QtCore.QObject):
 
         # Set the UI object
         self.ui = QtWidgets.QDialog()
-        self.ui.resize(600, 800)
+        self.ui.resize(650, 800)
         dialog_layout = QtWidgets.QVBoxLayout(self.ui)
         dialog_layout.setContentsMargins(2, 2, 2, 2)
         self.ui.setWindowTitle("Settings")
