Switchboard: Default nDisplay launches to -NoLoadingScreen to avoid blackouts when loading new maps or saving them from Multi-User.

An option has been added in case the user wants to suppress this option.

diff --git a/Engine/Plugins/VirtualProduction/Switchboard/Source/Switchboard/switchboard/devices/ndisplay/plugin_ndisplay.py b/Engine/Plugins/VirtualProduction/Switchboard/Source/Switchboard/switchboard/devices/ndisplay/plugin_ndisplay.py
--- a/Engine/Plugins/VirtualProduction/Switchboard/Source/Switchboard/switchboard/devices/ndisplay/plugin_ndisplay.py
+++ b/Engine/Plugins/VirtualProduction/Switchboard/Source/Switchboard/switchboard/devices/ndisplay/plugin_ndisplay.py
@@ -316,6 +316,12 @@ class DevicenDisplay(DeviceUnreal):
             nice_name="Sound",
             value=False,
         ),
+        'loading_screen': BoolSetting(
+            attr_name="loading_screen",
+            nice_name="Loading Screen",
+            value=False,
+            tool_tip='When unchecked, will add -NoLoadingScreen to the command line'
+        ),
         'render_api': OptionSetting(
             attr_name="render_api",
             nice_name="Render API",
@@ -835,6 +841,12 @@ class DevicenDisplay(DeviceUnreal):
             if not DevicenDisplay.csettings['sound'].get_value(self.name)
             else "")
 
+        # Loading Screen
+        no_loading_screen = (
+            "-NoLoadingScreen"
+            if not DevicenDisplay.csettings['loading_screen'].get_value(self.name)
+            else "")
+
         # MaxGPUCount (mGPU)
         max_gpu_count = DevicenDisplay.csettings["max_gpu_count"].get_value(
             self.name)
@@ -951,6 +963,7 @@ class DevicenDisplay(DeviceUnreal):
             f'{use_all_cores}',           # -useallavailablecores
             f'{no_texture_streaming}',    # -notexturestreaming
             f'{no_sound}',                # -nosound
+            f'{no_loading_screen}',       # -NoLoadingScreen
             f'-dc_node={self.name}',      # name of this node in the nDisplay cluster
             f'Log={self.log_filename}',   # log file
             f'{ini_engine}',              # Engine ini injections
