Switchboard: Sync fixes.

* Deliberately breaks compatibility with UGS for the very narrow case of CompatibleChangelist determination in a hotfix stream. Instead of freezing to the .0 release CL, we now always consider the most recent code change to be the CompatibleChangelist. The hotfix compatible CL is not pertinent in a stage environment, and prevents us from easily determining if we've actually built the CL we intended to sync to.

* The `p4_changelist_inspection.determine_changelist_type` function incorrectly considered any file extensions (or extensionless names) ending in the letters 'c' or 'h' to constitute code changes, and several other extensions were incorrectly ignored.

* Limits the progress logging interval from the sync process, preventing SwitchboardListener from trying to send potentially dozens of megabytes of stdout at the end of a sync, which was occasionally causing disconnects.

* Finally, fixed some linter warnings: shadowing built-in names (`id` and `dir`), opening files without specifying the encoding, `logging.warn` -> `logging.warning`.

diff --git a/Engine/Plugins/VirtualProduction/Switchboard/Source/Switchboard/sb_setup.py b/Engine/Plugins/VirtualProduction/Switchboard/Source/Switchboard/sb_setup.py
--- a/Engine/Plugins/VirtualProduction/Switchboard/Source/Switchboard/sb_setup.py
+++ b/Engine/Plugins/VirtualProduction/Switchboard/Source/Switchboard/sb_setup.py
@@ -83,7 +83,7 @@ class SbEnvBuilder(venv.EnvBuilder):
 
 class SbSetup:
     @classmethod
-    def build_parser(cls):
+    def build_parser(cls) -> argparse.ArgumentParser:
         parser = argparse.ArgumentParser()
 
         # Global options
@@ -122,13 +122,13 @@ class SbSetup:
 
         if not (CWRSYNC_FSTAB_PATH).exists():
             logging.info('Writing cwrsync fstab')
-            with open(CWRSYNC_FSTAB_PATH, 'wt') as fstab:
+            with open(CWRSYNC_FSTAB_PATH, 'wt', encoding='utf-8') as fstab:
                 fstab.write(CWRSYNC_FSTAB_CONTENTS)
 
     @staticmethod
-    def is_venv(dir: pathlib.Path) -> bool:
+    def is_venv(checkdir: pathlib.Path) -> bool:
         ''' Conservative estimate of whether `dir` contains a Python venv. '''
-        if not dir.is_dir():
+        if not checkdir.is_dir():
             return False
 
         # These lowercase names are canonical for non-Windows OSes, while
@@ -143,7 +143,7 @@ class SbSetup:
         existing_files: Set[str] = set()
         existing_dirs: Set[str] = set()
 
-        with os.scandir(dir) as scan:
+        with os.scandir(checkdir) as scan:
             for entry in scan:
                 name: str = entry.name
                 if sys.platform.startswith('win'):
@@ -266,7 +266,8 @@ class SbSetup:
 
         if options.action == 'install':
             install_result = self.run_install(options)
-            with open(options.venv_dir / 'provision.log', 'wt+') as logfile:
+            with open(options.venv_dir / 'provision.log', 'wt+',
+                      encoding='utf-8') as logfile:
                 logfile.write(log_buffer.getvalue())
             return install_result
         elif options.action == 'verify':
diff --git a/Engine/Plugins/VirtualProduction/Switchboard/Source/Switchboard/sbl_helper.py b/Engine/Plugins/VirtualProduction/Switchboard/Source/Switchboard/sbl_helper.py
--- a/Engine/Plugins/VirtualProduction/Switchboard/Source/Switchboard/sbl_helper.py
+++ b/Engine/Plugins/VirtualProduction/Switchboard/Source/Switchboard/sbl_helper.py
@@ -496,6 +496,7 @@ class SbListenerHelper:
         self.clobber_engine: bool = False
         self.clobber_project: bool = False
         self.generate_after_sync: bool = False
+        self.ugs_versioning: bool = False
 
         self.dry_run: bool = False
         self.p4user: Optional[str] = None
@@ -544,6 +545,9 @@ class SbListenerHelper:
         sync_parser.add_argument(
             '--clobber-project', action='store_true',
             help='Override noclobber for project files')
+        sync_parser.add_argument(
+            '--ugs-versioning', action='store_true',
+            help='Maintain depot CompatibleChangelist (e.g. hotfix stream)')
 
         sync_parser.add_argument('-n', '--dry-run', action='store_true',
                                  help='Preview sync; do not update files')
@@ -625,13 +629,13 @@ class SbListenerHelper:
                     uuid.UUID(id)
                     for id in options.include_categories.split(','))
                 self.ugs_filters.read_categories_from_ini_parser(ugs_config)
-                for category in self.ugs_filters.categories.values():
-                    if category.id not in include_ids:
-                        logging.debug(f'Excluding sync filter "{category.name}"')
-                        self.ugs_filters.exclude_category(category.id)
+                for cat in self.ugs_filters.categories.values():
+                    if cat.catid not in include_ids:
+                        logging.debug(f'Excluding sync filter "{cat.name}"')
+                        self.ugs_filters.exclude_category(cat.catid)
                         any_active_filters = True
                     else:
-                        include_ids.remove(category.id)
+                        include_ids.remove(cat.catid)
 
                 # Any remaining IDs here were unmatched to categories
                 if len(include_ids):
@@ -931,11 +935,16 @@ class SbListenerHelper:
             'BranchName': branch_name.replace('/', '+'),
         }
 
-        # Don't overwrite the compatible changelist if we're in a hotfix
-        # release.
-        no_depot_compat = build_ver['CompatibleChangelist'] == 0
-        licensee_changed = (build_ver['IsLicenseeVersion']
-                            != updates['IsLicenseeVersion'])
+        if self.ugs_versioning:
+            # Don't overwrite the compatible changelist if we're in a hotfix
+            # release stream.
+            no_depot_compat = build_ver['CompatibleChangelist'] == 0
+            licensee_changed = (build_ver['IsLicenseeVersion']
+                                != updates['IsLicenseeVersion'])
+        else:
+            # Always overwrite compatible changelist.
+            no_depot_compat = True
+
         if no_depot_compat or licensee_changed:
             updates['CompatibleChangelist'] = ver_compatible_cl
 
@@ -980,7 +989,14 @@ class SbListenerHelper:
             end = '' if progress_pct != 1.0 else None
             print(f'\r{" "*70}\r{progress_str}\r', end=end, flush=True)
         else:
-            logging.info(progress_str)
+            # Limit number of progress lines; otherwise one per file
+            MAX_PROGRESS_LINES = 5000
+            progress_interval = max(1, total_files // MAX_PROGRESS_LINES)
+            if (
+                (completed_files % progress_interval == 0) or
+                (completed_files == total_files)
+            ):
+                logging.info(progress_str)
 
     def check_sync_options(self, options: argparse.Namespace):
         '''
@@ -990,6 +1006,7 @@ class SbListenerHelper:
         self.dry_run = options.dry_run
         self.clobber_engine = options.clobber_engine
         self.clobber_project = options.clobber_project
+        self.ugs_versioning = options.ugs_versioning
         self.ugs_lib_dir = options.ugs_lib_dir
 
         self.p4user = options.p4user
diff --git a/Engine/Plugins/VirtualProduction/Switchboard/Source/Switchboard/switchboard/devices/unreal/plugin_unreal.py b/Engine/Plugins/VirtualProduction/Switchboard/Source/Switchboard/switchboard/devices/unreal/plugin_unreal.py
--- a/Engine/Plugins/VirtualProduction/Switchboard/Source/Switchboard/switchboard/devices/unreal/plugin_unreal.py
+++ b/Engine/Plugins/VirtualProduction/Switchboard/Source/Switchboard/switchboard/devices/unreal/plugin_unreal.py
@@ -788,7 +788,7 @@ class DeviceUnreal(Device):
         include_setting = DeviceUnreal.csettings['included_sync_categories']
         include_setting.possible_values.clear()
         for category in sync_filters.categories.values():
-            option = SyncCategoryOption(category.id, category.name)
+            option = SyncCategoryOption(category.catid, category.name)
             include_setting.possible_values.append(option)
 
     @classmethod
diff --git a/Engine/Plugins/VirtualProduction/Switchboard/Source/Switchboard/switchboard/ugs_utils.py b/Engine/Plugins/VirtualProduction/Switchboard/Source/Switchboard/switchboard/ugs_utils.py
--- a/Engine/Plugins/VirtualProduction/Switchboard/Source/Switchboard/switchboard/ugs_utils.py
+++ b/Engine/Plugins/VirtualProduction/Switchboard/Source/Switchboard/switchboard/ugs_utils.py
@@ -275,7 +275,7 @@ def run(
         logging.info('Restoring previous UnrealGameSync state.')
         if _set_active_ugs_context(ugs_bin_path,
                                    ugs_state_to_restore, cwd=cwd) != 0:
-            logging.warn("Failed to restore UnrealGameSync's state.")
+            logging.warning("Failed to restore UnrealGameSync's state.")
 
     return ugs_proc.returncode
 
@@ -507,8 +507,8 @@ def parse_depot_ugs_configs(
             if (b'clientFile' in record) and (b'action' in record):
                 # If the client file exists and is modified, use local ver
                 client_path = pathlib.Path(record[b'clientFile'].decode())
-                if client_path.exists and action == b'edit':
-                    local_contents = client_path.read_text()
+                if client_path.exists() and action == b'edit':
+                    local_contents = client_path.read_text(encoding='utf-8')
                     config_contents_map[str(client_path)] = local_contents
                     continue
                 elif action == b'delete':
@@ -577,8 +577,8 @@ class SyncFilters:
         skipinitialspace = True
 
     class Category:
-        def __init__(self, id: uuid.UUID, name: str, paths: list[str]):
-            self.id = id
+        def __init__(self, catid: uuid.UUID, name: str, paths: list[str]):
+            self.catid = catid
             self.name = name
             self.paths = paths
 
@@ -603,16 +603,16 @@ class SyncFilters:
             sfile = io.StringIO(categories_raw)
             self._read(sfile)
 
-    def exclude_category(self, category_id: uuid.UUID):
+    def exclude_category(self, catid: uuid.UUID):
         if not P4PYTHON_AVAILABLE:
             raise RuntimeError('This method currently requires P4PYTHON')
 
-        if category_id in self.excluded_categories:
+        if catid in self.excluded_categories:
             return  # already excluded
 
-        category = self.categories.get(category_id)
+        category = self.categories.get(catid)
         if category is None:
-            logging.warning(f'exclude_category: no such category {category_id}')
+            logging.warning(f'exclude_category: no such category {catid}')
             return
 
         self._map_exclusion(category)
@@ -633,7 +633,7 @@ class SyncFilters:
                 value = value.strip('"')
                 category_dict[key] = value
 
-            id = uuid.UUID(hex=category_dict['UniqueId'])
+            catid = uuid.UUID(hex=category_dict['UniqueId'])
             name = category_dict['Name']
             paths = category_dict['Paths'].split(';')
             paths = [path.lstrip() for path in paths]
@@ -642,23 +642,23 @@ class SyncFilters:
             requires = category_dict.get('Requires')
             requires = uuid.UUID(requires) if requires else None
 
-            if id not in self.categories:
-                category = self.create_category(id, name, paths)
+            if catid not in self.categories:
+                category = self.create_category(catid, name, paths)
                 category.hidden = hidden
                 category.requires = requires
             else:
                 # This seems weird, but... some configs need it
-                self.categories[id].paths.extend(paths)
+                self.categories[catid].paths.extend(paths)
 
-    def create_category(self, id: uuid.UUID, name: str, paths: list[str]):
-        category = SyncFilters.Category(id, name, paths)
-        self.categories[id] = category
+    def create_category(self, catid: uuid.UUID, name: str, paths: list[str]):
+        category = SyncFilters.Category(catid, name, paths)
+        self.categories[catid] = category
         return category
 
     DIR_REGEX = re.compile(r'(\.\.\.|/|\\)')
 
     def _map_exclusion(self, category: Category):
-        self.excluded_categories.add(category.id)
+        self.excluded_categories.add(category.catid)
         for path in category.paths:
             self.map.insert(f'-{path}')
 
diff --git a/Engine/Plugins/VirtualProduction/Switchboard/Source/Switchboard/switchboard/util/p4_changelist_inspection.py b/Engine/Plugins/VirtualProduction/Switchboard/Source/Switchboard/switchboard/util/p4_changelist_inspection.py
--- a/Engine/Plugins/VirtualProduction/Switchboard/Source/Switchboard/switchboard/util/p4_changelist_inspection.py
+++ b/Engine/Plugins/VirtualProduction/Switchboard/Source/Switchboard/switchboard/util/p4_changelist_inspection.py
@@ -5,7 +5,8 @@ import typing
 from enum import IntFlag
 
 from ..config import CONFIG
-from ..p4_utils import run, valueForMarshalledKey, hasValueForMarshalledKey
+from ..p4_utils import (run, valueForMarshalledKey, hasValueForMarshalledKey,
+                        ALL_CODE_EXTS)
 from ..switchboard_utils import get_hidden_sp_startupinfo
 
 
@@ -59,11 +60,9 @@ class EChangelistFileFlags(IntFlag):
     
     
 def determine_changelist_type(changelist: int) -> EChangelistFileFlags:
-    compiled_ue_source_extensions = ["c", "cs", "cpp", "h", "hpp"]
-    
     changes: dict = run(f"describe {changelist}")
     _raise_exception_on_error(changes)
-    
+
     has_source_changes = False
     has_content_changes = False
     for change in changes:
@@ -71,17 +70,19 @@ def determine_changelist_type(changelist: int) -> EChangelistFileFlags:
         # Each changed file path name will have an entry depotFilex
         while hasValueForMarshalledKey(change, f"depotFile{file_index}"):
             file_path_name: str = valueForMarshalledKey(change, f"depotFile{file_index}")
-            
-            has_source_changes = any([file_path_name.endswith(extension) for extension in compiled_ue_source_extensions])
+
+            has_source_changes = any(
+                [file_path_name.endswith(ext) for ext in ALL_CODE_EXTS])
+
             # If it's not a source file, treat it as a content file
             has_content_changes = not has_source_changes
-            
+
             # All flags fulfilled - no need to inspect more files
             if has_source_changes and has_content_changes:
                 break
-            
+
             file_index += 1
-    
+
     result = 0
     if has_source_changes:
         result |= EChangelistFileFlags.Code
