Switchboard: Fix for terminating multi-user server on build.

Fixes regression from CL 24476690.

diff --git a/Engine/Plugins/VirtualProduction/Switchboard/Source/Switchboard/switchboard/devices/unreal/plugin_unreal.py b/Engine/Plugins/VirtualProduction/Switchboard/Source/Switchboard/switchboard/devices/unreal/plugin_unreal.py
--- a/Engine/Plugins/VirtualProduction/Switchboard/Source/Switchboard/switchboard/devices/unreal/plugin_unreal.py
+++ b/Engine/Plugins/VirtualProduction/Switchboard/Source/Switchboard/switchboard/devices/unreal/plugin_unreal.py
@@ -1670,7 +1670,7 @@ class DeviceUnreal(Device):
                         QtWidgets.QMessageBox.Yes | QtWidgets.QMessageBox.No)
 
                     if mb_ret == QtWidgets.QMessageBox.Yes:
-                        DeviceUnreal.mu_server.terminate(bypolling=True)
+                        DeviceUnreal.mu_server.process.kill()
 
                 puuid_dependency = self._build_mu_server(
                     puuid_dependency=puuid_dependency)
diff --git a/Engine/Plugins/VirtualProduction/Switchboard/Source/Switchboard/switchboard/switchboard_application.py b/Engine/Plugins/VirtualProduction/Switchboard/Source/Switchboard/switchboard/switchboard_application.py
--- a/Engine/Plugins/VirtualProduction/Switchboard/Source/Switchboard/switchboard/switchboard_application.py
+++ b/Engine/Plugins/VirtualProduction/Switchboard/Source/Switchboard/switchboard/switchboard_application.py
@@ -277,7 +277,6 @@ class MultiUserApplication:
         configured_server_name = CONFIG.MUSERVER_SERVER_NAME.get_value()
         return f'{configured_server_name}_{encoded}'
 
-
     def clear_process_info(self):
         self._running_endpoint = ""
         self._running_name = ""
@@ -291,6 +290,9 @@ class MultiUserApplication:
         self.clear_process_info()
         return False
 
+    def kill(self):
+        self.process.kill()
+
 
 MultiUserServerInstance: Optional[MultiUserApplication] = None
 
