Switchboard: Include CLs from AdditionalPathsToSync in project CL dropdown, and add CL descriptions to tooltips.

Also fixes an incorrect QIcon argument (unqualified filename) to make it qrc-relative, and fixes some linter warnings about parameter names shadowing `range()`.

diff --git a/Engine/Plugins/VirtualProduction/Switchboard/Source/Switchboard/sbl_helper.py b/Engine/Plugins/VirtualProduction/Switchboard/Source/Switchboard/sbl_helper.py
--- a/Engine/Plugins/VirtualProduction/Switchboard/Source/Switchboard/sbl_helper.py
+++ b/Engine/Plugins/VirtualProduction/Switchboard/Source/Switchboard/sbl_helper.py
@@ -872,7 +872,7 @@ class SbListenerHelper:
             if quick_lower > 0:
                 quick_range = f'{quick_lower},{upper_cl}'
                 check_lower = p4_utils.p4_latest_code_change(
-                    paths, range=quick_range, exts=['.cpp'], **p4opts)
+                    paths, in_range=quick_range, exts=['.cpp'], **p4opts)
 
                 if check_lower:
                     if in_lower_cl:
@@ -903,14 +903,14 @@ class SbListenerHelper:
             eng_code_cl = eng_upper_bound_cl
         else:
             eng_code_cl = p4_utils.p4_latest_code_change(
-                engine_paths, range=eng_range, **p4opts)
+                engine_paths, in_range=eng_range, **p4opts)
 
         if proj_upper_bound_cl:
             if lower_bound_cl and (lower_bound_cl >= proj_upper_bound_cl):
                 proj_code_cl = proj_upper_bound_cl
             else:
                 proj_code_cl = p4_utils.p4_latest_code_change(
-                    project_paths, range=proj_range, **p4opts)
+                    project_paths, in_range=proj_range, **p4opts)
         else:
             proj_code_cl = None
 
diff --git a/Engine/Plugins/VirtualProduction/Switchboard/Source/Switchboard/switchboard/devices/unreal/plugin_unreal.py b/Engine/Plugins/VirtualProduction/Switchboard/Source/Switchboard/switchboard/devices/unreal/plugin_unreal.py
--- a/Engine/Plugins/VirtualProduction/Switchboard/Source/Switchboard/switchboard/devices/unreal/plugin_unreal.py
+++ b/Engine/Plugins/VirtualProduction/Switchboard/Source/Switchboard/switchboard/devices/unreal/plugin_unreal.py
@@ -697,6 +697,12 @@ class DeviceUnreal(Device):
                 "On Windows, the default install location is '${LOCALAPPDATA}/UnrealGameSync/Latest/'.")
         )
 
+    # A stand-in QObject to emit signals from classmethods.
+    class StaticSignals(QtCore.QObject):
+        ugs_config_updated_signal = QtCore.Signal(ugs_utils.IniParser)
+
+    static_signals = StaticSignals()
+
     unreal_started_signal = QtCore.Signal()
 
     mu_server = switchboard_application.get_multi_user_server_instance()
@@ -791,6 +797,8 @@ class DeviceUnreal(Device):
             option = SyncCategoryOption(category.catid, category.name)
             include_setting.possible_values.append(option)
 
+        cls.static_signals.ugs_config_updated_signal.emit(ugs_config)
+
     @classmethod
     def get_designated_local_builder(cls) -> Optional[DeviceUnreal]:
         '''
diff --git a/Engine/Plugins/VirtualProduction/Switchboard/Source/Switchboard/switchboard/p4_utils.py b/Engine/Plugins/VirtualProduction/Switchboard/Source/Switchboard/switchboard/p4_utils.py
--- a/Engine/Plugins/VirtualProduction/Switchboard/Source/Switchboard/switchboard/p4_utils.py
+++ b/Engine/Plugins/VirtualProduction/Switchboard/Source/Switchboard/switchboard/p4_utils.py
@@ -14,7 +14,7 @@ from . import switchboard_utils as sb_utils
 from .devices.unreal.version_helpers import LISTENER_COMPATIBLE_VERSION
 
 
-# Metadata attached to p4 commands for sever analytics
+# Metadata attached to p4 commands for server analytics
 meta_zprog = 'switchboard'
 meta_zversion = '.'.join(str(x) for x in LISTENER_COMPATIBLE_VERSION)  # FIXME?
 
@@ -33,19 +33,30 @@ def p4_login(f):
 
 
 @p4_login
-def p4_latest_changelist(p4_path, working_dir, num_changelists=10):
-    """
-    Return (num_changelists) latest CLs
-    """
-    p4_command = f'p4 -ztag -F "%change%" changes -m {num_changelists} {p4_path}/...'
-    logging.info(f"Executing: {p4_command}")
+def p4_latest_changelists(
+    p4_paths: str | list[str],
+    *,
+    limit: int = 10,
+    exclude_automation: bool = True,
+    user: Optional[str] = None,
+    client: Optional[str] = None,
+) -> list[tuple[int, str]]:
+    ''' Return (cl: int, description: str) for `limit` most recent CLs '''
+    if isinstance(p4_paths, str):
+        p4_paths = [p4_paths]
+
+    changes = p4_changes(p4_paths, limit=limit, user=user, client=client)
+
+    if exclude_automation:
+        changes = filter(lambda x: x[b'user'] != b'buildmachine', changes)
+
+    # TODO: Support UnrealGameSync.ini [Options] ExcludeChanges regexes?
 
-    p4_result = subprocess.check_output(p4_command, cwd=working_dir, shell=True, startupinfo=sb_utils.get_hidden_sp_startupinfo()).decode()
+    cl_desc_pairs = [(int(change[b'change']), change[b'desc'].decode())
+                     for change in changes]
 
-    if p4_result:
-        return p4_result.split()
+    return sorted(cl_desc_pairs, reverse=True)[:limit]
 
-    return None
 
 def run(cmd, args=[], input=None):
     ''' Runs the provided p4 command and arguments with -G python marshaling 
@@ -370,7 +381,7 @@ ALL_CODE_EXTS = [
 def p4_latest_code_change(
     paths: list[str],
     *,
-    range: Optional[str] = None,
+    in_range: Optional[str] = None,
     exts: Optional[list[str]] = None,
     user: Optional[str] = None,
     client: Optional[str] = None,
@@ -382,7 +393,7 @@ def p4_latest_code_change(
     if exts is None:
         exts = ALL_CODE_EXTS
 
-    rangespec = f'@{range}' if range is not None else ''
+    rangespec = f'@{in_range}' if in_range is not None else ''
 
     code_paths: list[str] = []
     for path in paths:
diff --git a/Engine/Plugins/VirtualProduction/Switchboard/Source/Switchboard/switchboard/switchboard_dialog.py b/Engine/Plugins/VirtualProduction/Switchboard/Source/Switchboard/switchboard/switchboard_dialog.py
--- a/Engine/Plugins/VirtualProduction/Switchboard/Source/Switchboard/switchboard/switchboard_dialog.py
+++ b/Engine/Plugins/VirtualProduction/Switchboard/Source/Switchboard/switchboard/switchboard_dialog.py
@@ -3,6 +3,7 @@
 from __future__ import annotations
 
 import datetime
+from itertools import count
 import logging
 import os
 import threading
@@ -342,7 +343,7 @@ class SwitchboardDialog(QtCore.QObject):
 
         self.refresh_levels_button = sb_widgets.ControlQPushButton()
         self.refresh_levels_button.setMaximumSize(22, 22)
-        self.refresh_levels_button.setIcon(QtGui.QIcon("icon_refresh.png"))
+        self.refresh_levels_button.setIcon(QtGui.QIcon(':/icons/images/icon_refresh.png'))
         self.refresh_levels_button.setProperty("frameless", True)
         self.refresh_levels_button.setToolTip("Refresh level list")
 
@@ -510,6 +511,10 @@ class SwitchboardDialog(QtCore.QObject):
         self.script_manager.on_postinit(self)
         self.have_warned_about_muserver = False
 
+        DeviceUnreal.static_signals.ugs_config_updated_signal.connect(
+            lambda _: self.p4_refresh_project_cl()
+        )
+
     def _try_change_address(self):
         new_value = self.window.current_address_value.text()
         old_address = SETTINGS.ADDRESS.get_value()
@@ -1934,7 +1939,9 @@ class SwitchboardDialog(QtCore.QObject):
 
         sync_method = CONFIG.ENGINE_SYNC_METHOD.get_value()
 
-        changelists = None
+        changelists: Optional[list[int]] = None
+        descriptions: Optional[list[str]] = None
+
         # If we're syncing 'Precompiled Binaries', then that implies that we should be using UGS:
         if ENABLE_UGS_SUPPORT:
             if sync_method == EngineSyncMethod.Sync_PCBs.value or sync_method == EngineSyncMethod.Sync_From_UGS.value:
@@ -1945,14 +1952,30 @@ class SwitchboardDialog(QtCore.QObject):
 
         if not changelists:
             LOGGER.info("Refreshing p4 project changelists")
-            working_dir = os.path.dirname(CONFIG.UPROJECT_PATH.get_value())
-            changelists = p4_utils.p4_latest_changelist(CONFIG.P4_PROJECT_PATH.get_value(), working_dir)
+            client = CONFIG.SOURCE_CONTROL_WORKSPACE.get_value()
+            paths = [f'{CONFIG.P4_PROJECT_PATH.get_value()}/...']
+
+            if DeviceUnreal.ugs_config:
+                if addpaths := DeviceUnreal.ugs_config.try_get(
+                    'Perforce', 'AdditionalPathsToSync'
+                ):
+                    paths.extend([f'//{client}{x}' for x in addpaths])
+
+            cl_descs = p4_utils.p4_latest_changelists(paths, client=client)
+            changelists = [str(cl[0]) for cl in cl_descs]
+            descriptions = [cl[1] for cl in cl_descs]
 
         self.window.project_cl_combo_box.clear()
 
         if changelists:
             self.window.project_cl_combo_box.addItems(changelists)
             self.window.project_cl_combo_box.setCurrentIndex(0)
+
+        if descriptions:
+            for (idx, desc) in zip(count(), descriptions):
+                self.window.project_cl_combo_box.setItemData(
+                    idx, desc, QtCore.Qt.ItemDataRole.ToolTipRole)
+
         self.window.project_cl_combo_box.addItem(EMPTY_SYNC_ENTRY)
 
     def p4_refresh_engine_cl(self):
@@ -1970,11 +1993,20 @@ class SwitchboardDialog(QtCore.QObject):
 
             engine_p4_path = CONFIG.P4_ENGINE_PATH.get_value()
             if engine_p4_path:
-                working_dir = os.path.dirname(CONFIG.UPROJECT_PATH.get_value())
-                changelists = p4_utils.p4_latest_changelist(engine_p4_path, working_dir)
+                client = CONFIG.SOURCE_CONTROL_WORKSPACE.get_value()
+                cl_descs = p4_utils.p4_latest_changelists(
+                    engine_p4_path+'/...', client=client)
+                changelists = [str(cl[0]) for cl in cl_descs]
+                descriptions = [cl[1] for cl in cl_descs]
+
                 if changelists:
                     self.window.engine_cl_combo_box.addItems(changelists)
                     self.window.engine_cl_combo_box.setCurrentIndex(0)
+
+                if descriptions:
+                    for (idx, desc) in zip(count(), descriptions):
+                        self.window.engine_cl_combo_box.setItemData(
+                            idx, desc, QtCore.Qt.ItemDataRole.ToolTipRole)
             else:
                 LOGGER.warning('"Build Engine" is enabled in the settings but the engine does not seem to be under perforce control.')
                 LOGGER.warning("Please check your perforce settings.")
