Switchboard: Fixed not handling opening a config that has paths to locked drive (e.g. bitlocker).

File operations may raise OSError in such case.

#rnx

diff --git a/Engine/Plugins/VirtualProduction/Switchboard/Source/Switchboard/switchboard/devices/unreal/listener_watcher.py b/Engine/Plugins/VirtualProduction/Switchboard/Source/Switchboard/switchboard/devices/unreal/listener_watcher.py
--- a/Engine/Plugins/VirtualProduction/Switchboard/Source/Switchboard/switchboard/devices/unreal/listener_watcher.py
+++ b/Engine/Plugins/VirtualProduction/Switchboard/Source/Switchboard/switchboard/devices/unreal/listener_watcher.py
@@ -98,6 +98,10 @@ class ListenerWatcher(QtCore.QObject):
             except subprocess.CalledProcessError as e:
                 LOGGER.error(f"Non-zero exit code ({e.returncode}) when invoking `{e.cmd}`", exc_info=e)
                 time.sleep(retry_interval_sec)
+            except OSError as e:
+                # This can happen if the drive is locked (e.g. by Bitlocker)
+                LOGGER.error(e)
+                break
 
         if not updated:
             self.listener_mtime = 0
diff --git a/Engine/Plugins/VirtualProduction/Switchboard/Source/Switchboard/switchboard/devices/unreal/plugin_unreal.py b/Engine/Plugins/VirtualProduction/Switchboard/Source/Switchboard/switchboard/devices/unreal/plugin_unreal.py
--- a/Engine/Plugins/VirtualProduction/Switchboard/Source/Switchboard/switchboard/devices/unreal/plugin_unreal.py
+++ b/Engine/Plugins/VirtualProduction/Switchboard/Source/Switchboard/switchboard/devices/unreal/plugin_unreal.py
@@ -2667,11 +2667,15 @@ class DeviceUnreal(Device):
         else:
             local_project_dir = pathlib.Path(
                 CONFIG.UPROJECT_PATH.get_value()).parent
-            if local_project_dir.is_dir():
-                log_download_dir = \
-                    local_project_dir / 'Saved' / 'Logs' / 'Switchboard'
-                log_download_dir.mkdir(parents=True, exist_ok=True)
-                return log_download_dir
+            try:
+                if local_project_dir.is_dir():
+                    log_download_dir = \
+                        local_project_dir / 'Saved' / 'Logs' / 'Switchboard'
+                    log_download_dir.mkdir(parents=True, exist_ok=True)
+                    return log_download_dir
+            except OSError as e:
+                # This can happen if the drive is locked (e.g. by Bitlocker)
+                LOGGER.error(e)
 
         return None
 
