Switchboard: Fix nDisplay node hitches that happen after SBL is left running for a few hours, by limiting NvApi calls that cause hitches to connction time instead of regularly.

We discovered that NvAPI calls can cause hitches in frame presentation, this CL adds fine control over what status data is polled and which should be explicitly requested by the user. This is used to request status data that does not change often (e.g. driver version) and also calls that have been determined to cause hitches when calling dxgi present, only at time of connection or when the users presses "Refresh Table", instead of regularly.

For the regular GPU stat calls, the NvAPI calls were replaced with NVMI calls, which do not seem to be causing hitches.

"Refresh Mosaics" button renamed to "Refresh Table" because it now refreshes more than just the mosaics. Now it also requests a full sync status poll. This intentional action might cause a hitch on the nodes but this is now triggered explicitly by the user, which is rarely needed anyway, instead of every second.

Also fixed an issue of the nDisplay monitor table not updating its display after emitting the dataChanged signal.

Fixed nDisplay Monitor not updating when the nDisplay node name or address changed.

Improved nDisplay monitor to show the connection icon instead of the wider column "Connected" with yes/no values.

Vertically aligned the cells in the table.

It now doesn't use the orange warning color on certain columns that are only relevant when the node is running (e.g. present mode).

Added a code-only option to disable exe flags button.

#rn Fixed nDisplay hitches that happen after a few hours of running SBL.

diff --git a/Engine/Plugins/VirtualProduction/Switchboard/Source/Switchboard/switchboard/config.py b/Engine/Plugins/VirtualProduction/Switchboard/Source/Switchboard/switchboard/config.py
--- a/Engine/Plugins/VirtualProduction/Switchboard/Source/Switchboard/switchboard/config.py
+++ b/Engine/Plugins/VirtualProduction/Switchboard/Source/Switchboard/switchboard/config.py
@@ -247,8 +247,9 @@ class Setting(QtCore.QObject):
     def on_device_name_changed(self, old_name: str, new_name: str):
         if old_name in self._overrides.keys():
             self._overrides[new_name] = self._overrides.pop(old_name)
-            self._override_widgets[new_name] = (
-                self._override_widgets.pop(old_name))
+
+        if old_name in self._override_widgets.keys():
+            self._override_widgets[new_name] = (self._override_widgets.pop(old_name))
 
     def reset(self):
         self._value = self._original_value
@@ -2691,10 +2692,12 @@ class Config(object):
                     setting.on_device_name_changed(old_name, new_name)
                 break
 
-        new_key = (old_key[0], new_name)
-        self._device_settings[new_key] = self._device_settings.pop(old_key)
-
-        self.save()
+        if old_key is not None:
+            new_key = (old_key[0], new_name)
+            self._device_settings[new_key] = self._device_settings.pop(old_key)
+            self.save()
+        else:
+            LOGGER.warning(f"Device with name '{old_name}' not found in the config's _device_settings")
 
     def on_device_removed(self, _, device_type, device_name, update_config):
         if not update_config:
diff --git a/Engine/Plugins/VirtualProduction/Switchboard/Source/Switchboard/switchboard/devices/ndisplay/ndisplay_monitor.py b/Engine/Plugins/VirtualProduction/Switchboard/Source/Switchboard/switchboard/devices/ndisplay/ndisplay_monitor.py
--- a/Engine/Plugins/VirtualProduction/Switchboard/Source/Switchboard/switchboard/devices/ndisplay/ndisplay_monitor.py
+++ b/Engine/Plugins/VirtualProduction/Switchboard/Source/Switchboard/switchboard/devices/ndisplay/ndisplay_monitor.py
@@ -7,10 +7,12 @@ import traceback
 
 from PySide6 import QtCore
 from PySide6.QtCore import QAbstractTableModel, QModelIndex, Qt, QTimer
-from PySide6.QtGui import QColor
+from PySide6.QtGui import QColor, QIcon
 
 from switchboard import message_protocol
+from switchboard.message_protocol import SyncStatusRequestFlags
 from switchboard.switchboard_logging import LOGGER
+from switchboard.devices.device_base import Device
 
 
 class nDisplayMonitor(QAbstractTableModel):
@@ -22,27 +24,32 @@ class nDisplayMonitor(QAbstractTableModel):
 
     COLOR_WARNING = QColor(0x70, 0x40, 0x00)
     COLOR_NORMAL = QColor(0x3d, 0x3d, 0x3d)
+    COLOR_DISCONNECTED = QColor(0x7F, 0x7F, 0x7F)
     CORE_OVERLOAD_THRESH = 90  # percent utilization
 
     # Special meaning when received from listener for PresentMode, but
     # sometimes also used for other values' display for consistency.
     DATA_MISSING = 'n/a'
 
+    # If false, the button to disable full screen optimizations is hidden
+    # and periodic polling of the state of this feature is disabled.
+    use_exe_flags = True
+
     def __init__(self, parent):
         QAbstractTableModel.__init__(self, parent)
 
-        self.polling_period_ms = 1000
+        self.polling_period_ms = 2000
 
         # ordered so that we can map row indices to devices
         self.devicedatas = OrderedDict()
 
         self.timer = QTimer(self)
-        self.timer.timeout.connect(self.poll_sync_status)
+        self.timer.timeout.connect(self.poll_variable_sync_status)
 
         HEADER_DATA = {
+            'Connected': 'If we are connected to the listener of this device',
             'Node': 'The cluster name of this device',
             'Host': 'The URL of the remote PC',
-            'Connected': 'If we are connected to the listener of this device',
             'Driver': 'GPU driver version',
             'PresentMode':
                 'Current presentation mode. Only available once the render '
@@ -84,10 +91,18 @@ class nDisplayMonitor(QAbstractTableModel):
                 'sensors.)',
         }
 
+        if not self.use_exe_flags:
+            HEADER_DATA.pop('ExeFlags')
+
         self.colnames = list(HEADER_DATA.keys())
         self.tooltips = list(HEADER_DATA.values())
 
-    def color_for_column(self, colname, value, data):
+        # Load connection status icons
+        self.icon_unconnected = QIcon(":/icons/images/status_blank_disabled.png")
+        self.icon_connected = QIcon(":/icons/images/status_cyan.png")
+        self.icon_running = QIcon(":/icons/images/status_orange.png")
+
+    def color_for_column(self, colname, value, data, is_program_running: bool):
         ''' Returns the background color for the given cell '''
         if data['Connected'].lower() == 'no':
             if colname == 'Connected':
@@ -95,7 +110,11 @@ class nDisplayMonitor(QAbstractTableModel):
             return self.COLOR_NORMAL
 
         if colname == 'PresentMode':
-            is_good = 'Hardware Composed: Independent Flip' in value
+            ok_values = [
+                'Hardware Composed: Independent Flip',
+                'Hardware: Independent Flip',
+            ]
+            is_good = (not is_program_running) or any(ok_value in value for ok_value in ok_values)
             return self.COLOR_NORMAL if is_good else self.COLOR_WARNING
 
         if colname == 'Gpus':
@@ -103,10 +122,11 @@ class nDisplayMonitor(QAbstractTableModel):
             return self.COLOR_NORMAL if is_synced else self.COLOR_WARNING
 
         if colname == 'InFocus':
-            return self.COLOR_NORMAL if 'yes' in value else self.COLOR_WARNING
+            is_good = (not is_program_running) or ('yes' in value)
+            return self.COLOR_NORMAL if is_good else self.COLOR_WARNING
 
         if colname == 'ExeFlags':
-            is_good = 'DISABLEDXMAXIMIZEDWINDOWEDMODE' in value
+            is_good = (not is_program_running) or ('DISABLEDXMAXIMIZEDWINDOWEDMODE' in value)
             return self.COLOR_NORMAL if is_good else self.COLOR_WARNING
 
         if colname == 'Displays':
@@ -234,8 +254,7 @@ class nDisplayMonitor(QAbstractTableModel):
 
         # notify the UI
         row = deviceIdx + 1
-        self.dataChanged.emit(self.createIndex(row, 1),
-                              self.createIndex(row, len(self.colnames)))
+        self.refresh_display_for_row(row)
 
     def handle_connection_change(self, devicedata, deviceIdx):
         '''
@@ -255,11 +274,44 @@ class nDisplayMonitor(QAbstractTableModel):
                 self.reset_device_data(device, data)
 
             row = deviceIdx + 1
-            self.dataChanged.emit(self.createIndex(row, 1),
-                                  self.createIndex(row, len(self.colnames)))
+            self.refresh_display_for_row(row)
+
+    def default_program_id(self):
+        ''' Default value for program id when unreal is not running.
+        Used for messages that need this argument'''
+
+        return '00000000-0000-0000-0000-000000000000'
+
+    def program_id_from_device(self, device: Device):
+        ''' Returns the program id of the running nDisplay unreal instance '''
+        try:
+            program_id = device.program_start_queue.running_puuids_named('unreal')[-1]
+        except IndexError:
+            program_id = self.default_program_id()
 
-    def poll_sync_status(self):
+        return program_id
+
+    def poll_sync_status_for_device(self, device: Device, request_flags: SyncStatusRequestFlags):
+        ''' Polls sync status for the given device '''
+
+        # no point in continuing if not connected to listener
+        if not (device.unreal_client.is_connected
+                and device.unreal_client.is_authenticated):
+            return
+
+        # create message
+
+        program_id = self.program_id_from_device(device)
+
+        _, msg = message_protocol.create_get_sync_status_message(
+            program_id, request_flags)
+
+        # send get sync status message
+        device.unreal_client.send_message(msg)
+
+    def poll_sync_status(self, request_flags: SyncStatusRequestFlags):
         ''' Polls sync status for all nDisplay devices '''
+
         for deviceIdx, devicedata in enumerate(self.devicedatas.values()):
             device = devicedata['device']
 
@@ -269,23 +321,25 @@ class nDisplayMonitor(QAbstractTableModel):
             # detect stale devices
             self.handle_stale_device(devicedata, deviceIdx)
 
-            # no point in continuing if not connected to listener
-            if not (device.unreal_client.is_connected
-                    and device.unreal_client.is_authenticated):
-                continue
+            # request status
+            self.poll_sync_status_for_device(device, request_flags)
 
-            # create message
-            try:
-                program_id = device.program_start_queue.running_puuids_named(
-                    'unreal')[-1]
-            except IndexError:
-                program_id = '00000000-0000-0000-0000-000000000000'
+    def poll_variable_sync_status(self):
+        ''' Poll sync status but only include items that may change and that are
+        not likely to cause hitches in the target machine. '''
 
-            _, msg = message_protocol.create_get_sync_status_message(
-                program_id)
+        request_flags = SyncStatusRequestFlags.all()
 
-            # send get sync status message
-            device.unreal_client.send_message(msg)
+        # We don't poll status that can cause hitches, or that do not typically change.
+
+        request_flags &= ~SyncStatusRequestFlags.SyncTopos  # <-- Causes hitches when presenting frames.
+        request_flags &= ~SyncStatusRequestFlags.MosaicTopos
+        request_flags &= ~SyncStatusRequestFlags.DriverInfo
+
+        if not self.use_exe_flags:
+            request_flags &= ~SyncStatusRequestFlags.ProgramLayers
+
+        self.poll_sync_status(request_flags)
 
     def devicedata_from_device(self, device):
         ''' Retrieves the devicedata and index for given device '''
@@ -301,209 +355,230 @@ class nDisplayMonitor(QAbstractTableModel):
         Populates model data with message contents, which comes from 'get sync
         data' command.
         '''
+
         data = devicedata['data']
         device = devicedata['device']
 
+        request_flags = SyncStatusRequestFlags(message['request_flags'])
+        sync_status = message['syncStatus']
+
         #
         # Sync Topology
         #
-        sync_status = message['syncStatus']
-        sync_topos = sync_status['syncTopos']
-
-        # Build list informing which Gpus in each Sync group are in sync
-        gpus = []
-
-        for sync_topo in sync_topos:
-            gpu_sync_oks = [gpu['bIsSynced'] for gpu in sync_topo['syncGpus']]
-            gpu_syncs = map(lambda x: "Synced" if x else 'Free', gpu_sync_oks)
-            gpus.append('%s' % (', '.join(gpu_syncs)))
-
-        data['Gpus'] = '\n'.join(gpus) if len(gpus) > 0 else self.DATA_MISSING
-
-        # Build list informing which Display in each Sync group are in sync.
-        displays = []
-
-        bpc_strings = {1: 6, 2: 8, 3: 10, 4: 12, 5: 16}
-
-        for sync_topo in sync_topos:
-            display_sync_states = [
-                f"{syncDisplay['syncState']}"
-                f"({bpc_strings.get(syncDisplay['bpc'], '??')}bpc)"
-                for syncDisplay in sync_topo['syncDisplays']]
-            displays.append(', '.join(display_sync_states))
-
-        if len(displays) > 0:
-            data['Displays'] = '\n'.join(displays)
-        else:
-            data['Displays'] = self.DATA_MISSING
-
-        # Build Fps
-        refresh_rates = \
-            [f"{syncTopo['syncStatusParams']['refreshRate']*1e-4:.3f}"
-                for syncTopo in sync_topos]
-
-        if len(refresh_rates) > 0:
-            data['Fps'] = '\n'.join(refresh_rates)
-        else:
-            data['Fps'] = self.DATA_MISSING
-
-        # Build House Sync
-        house_fpss = [syncTopo['syncStatusParams']['houseSyncIncoming']*1e-4
-                      for syncTopo in sync_topos]
-        house_syncs = [syncTopo['syncStatusParams']['bHouseSync']
-                       for syncTopo in sync_topos]
-        house_sync_fpss = list(map(lambda x: f"{x[1]:.3f}" if x[0] else 'no',
-                               zip(house_syncs, house_fpss)))
-
-        if len(house_sync_fpss) > 0:
-            data['HouseSync'] = '\n'.join(house_sync_fpss)
-        else:
-            data['HouseSync'] = self.DATA_MISSING
-
-        # Build Sync Source
-        source_str = {0: 'Vsync', 1: 'House'}
-        sync_sources = [sync_topo['syncControlParams']['source']
-                        for sync_topo in sync_topos]
-        sync_sources = [source_str.get(sync_source, 'Unknown')
-                        for sync_source in sync_sources]
-        bInternalSecondaries = [
-            sync_topo['syncStatusParams']['bInternalSecondary']
-            for sync_topo in sync_topos]
-
-        sync_followers = []
-
-        for i in range(len(sync_sources)):
-            if bInternalSecondaries[i] and sync_sources[i] == 'Vsync':
-                sync_followers.append('Vsync(daisy)')
+        if SyncStatusRequestFlags.SyncTopos in request_flags:
+            sync_topos = sync_status['syncTopos']
+
+            # Build list informing which Gpus in each Sync group are in sync
+            gpus = []
+
+            for sync_topo in sync_topos:
+                gpu_sync_oks = [gpu['bIsSynced'] for gpu in sync_topo['syncGpus']]
+                gpu_syncs = map(lambda x: "Synced" if x else 'Free', gpu_sync_oks)
+                gpus.append('%s' % (', '.join(gpu_syncs)))
+
+            data['Gpus'] = '\n'.join(gpus) if len(gpus) > 0 else self.DATA_MISSING
+
+            # Build list informing which Display in each Sync group are in sync.
+            displays = []
+
+            bpc_strings = {1: 6, 2: 8, 3: 10, 4: 12, 5: 16}
+
+            for sync_topo in sync_topos:
+                display_sync_states = [
+                    f"{syncDisplay['syncState']}"
+                    f"({bpc_strings.get(syncDisplay['bpc'], '??')}bpc)"
+                    for syncDisplay in sync_topo['syncDisplays']]
+                displays.append(', '.join(display_sync_states))
+
+            if len(displays) > 0:
+                data['Displays'] = '\n'.join(displays)
             else:
-                sync_followers.append(sync_sources[i])
+                data['Displays'] = self.DATA_MISSING
 
-        if len(sync_followers) > 0:
-            data['SyncSource'] = '\n'.join(sync_followers)
-        else:
-            data['SyncSource'] = self.DATA_MISSING
+            # Build Fps
+            refresh_rates = \
+                [f"{syncTopo['syncStatusParams']['refreshRate']*1e-4:.3f}"
+                    for syncTopo in sync_topos]
 
+            if len(refresh_rates) > 0:
+                data['Fps'] = '\n'.join(refresh_rates)
+            else:
+                data['Fps'] = self.DATA_MISSING
+
+            # Build House Sync
+            house_fpss = [syncTopo['syncStatusParams']['houseSyncIncoming']*1e-4
+                          for syncTopo in sync_topos]
+            house_syncs = [syncTopo['syncStatusParams']['bHouseSync']
+                           for syncTopo in sync_topos]
+            house_sync_fpss = list(
+                map(
+                    lambda x: f"{x[1]:.3f}" if x[0] else 'no',
+                    zip(house_syncs, house_fpss)
+                )
+            )
+
+            if len(house_sync_fpss) > 0:
+                data['HouseSync'] = '\n'.join(house_sync_fpss)
+            else:
+                data['HouseSync'] = self.DATA_MISSING
+
+            # Build Sync Source
+            source_str = {0: 'Vsync', 1: 'House'}
+            sync_sources = [sync_topo['syncControlParams']['source']
+                            for sync_topo in sync_topos]
+            sync_sources = [source_str.get(sync_source, 'Unknown')
+                            for sync_source in sync_sources]
+            bInternalSecondaries = [
+                sync_topo['syncStatusParams']['bInternalSecondary']
+                for sync_topo in sync_topos]
+
+            sync_followers = []
+
+            for i in range(len(sync_sources)):
+                if bInternalSecondaries[i] and sync_sources[i] == 'Vsync':
+                    sync_followers.append('Vsync(daisy)')
+                else:
+                    sync_followers.append(sync_sources[i])
+
+            if len(sync_followers) > 0:
+                data['SyncSource'] = '\n'.join(sync_followers)
+            else:
+                data['SyncSource'] = self.DATA_MISSING
+
+        #
         # Mosaic Topology
-        mosaic_topos = sync_status['mosaicTopos']
+        #
+        if SyncStatusRequestFlags.MosaicTopos in request_flags:
+            mosaic_topos = sync_status['mosaicTopos']
 
-        mosaic_topo_lines = []
+            mosaic_topo_lines = []
 
-        for mosaic_topo in mosaic_topos:
-            display_settings = mosaic_topo['displaySettings']
-            width_per_display = display_settings['width']
-            height_per_display = display_settings['height']
+            for mosaic_topo in mosaic_topos:
+                display_settings = mosaic_topo['displaySettings']
+                width_per_display = display_settings['width']
+                height_per_display = display_settings['height']
 
-            width = mosaic_topo['columns'] * width_per_display
-            height = mosaic_topo['rows'] * height_per_display
+                width = mosaic_topo['columns'] * width_per_display
+                height = mosaic_topo['rows'] * height_per_display
 
-            # Ignoring displaySettings['freq'] because it seems to be fixed
-            # and ignores sync frequency.
-            line = f"{width}x{height} {display_settings['bpp']}bpp"
-            mosaic_topo_lines.append(line)
+                # Ignoring displaySettings['freq'] because it seems to be fixed
+                # and ignores sync frequency.
+                line = f"{width}x{height} {display_settings['bpp']}bpp"
+                mosaic_topo_lines.append(line)
 
-        data['Mosaics'] = '\n'.join(mosaic_topo_lines)
+            data['Mosaics'] = '\n'.join(mosaic_topo_lines)
 
-        # Build PresentMode.
-        flip_history = sync_status['flipModeHistory']
+        #
+        # Build PresentMode
+        #
+        if SyncStatusRequestFlags.FlipModeHistory in request_flags:
+            flip_history = sync_status['flipModeHistory']
 
-        if len(flip_history) > 0:
-            data['PresentMode'] = flip_history[-1]
+            if len(flip_history) > 0:
+                data['PresentMode'] = flip_history[-1]
 
-        # Detect PresentMode glitches
-        if len(set(flip_history)) > 1:
-            data['PresentMode'] = 'GLITCH!'
-            data['TimeLastFlipGlitch'] = time.time()
+            # Detect PresentMode glitches
+            if len(set(flip_history)) > 1:
+                data['PresentMode'] = 'GLITCH!'
+                data['TimeLastFlipGlitch'] = time.time()
 
-        # Write time since last glitch
-        if data['PresentMode'] != self.DATA_MISSING:
-            time_since_flip_glitch = time.time() - data['TimeLastFlipGlitch']
+            # Write time since last glitch
+            if data['PresentMode'] != self.DATA_MISSING:
+                time_since_flip_glitch = time.time() - data['TimeLastFlipGlitch']
 
-            # For 1 minute, let the user know that there was a flip mode glitch
-            if time_since_flip_glitch < 1*60:
-                data['PresentMode'] = data['PresentMode'].split('\n')[0] \
-                    + '\n' + str(int(time_since_flip_glitch))
+                # For 1 minute, let the user know that there was a flip mode glitch
+                if time_since_flip_glitch < 1*60:
+                    data['PresentMode'] = data['PresentMode'].split('\n')[0] \
+                        + '\n' + str(int(time_since_flip_glitch))
 
         # Window in focus or not
-        data['InFocus'] = 'no'
-        for prg in device.program_start_queue.running_programs_named('unreal'):
-            if prg.pid and prg.pid == sync_status['pidInFocus']:
-                data['InFocus'] = 'yes'
-                break
+        if SyncStatusRequestFlags.PidInFocus in request_flags:
+            data['InFocus'] = 'no'
+            for prg in device.program_start_queue.running_programs_named('unreal'):
+                if prg.pid and prg.pid == sync_status['pidInFocus']:
+                    data['InFocus'] = 'yes'
+                    break
 
         # Show Exe flags (like Disable Fullscreen Optimization)
-        data['ExeFlags'] = '\n'.join([
-            layer for layer in sync_status['programLayers'][1:]])
+        if SyncStatusRequestFlags.ProgramLayers in request_flags:
+            data['ExeFlags'] = '\n'.join([
+                layer for layer in sync_status['programLayers'][1:]])
 
         # Driver version
-        try:
-            driver = sync_status['driverVersion']
-            data['Driver'] = f'{int(driver/100)}.{driver % 100}'
-        except (KeyError, TypeError):
-            data['Driver'] = self.DATA_MISSING
+        if SyncStatusRequestFlags.DriverInfo in request_flags:
+            try:
+                driver = sync_status['driverVersion']
+                data['Driver'] = f'{int(driver/100)}.{driver % 100}'
+            except (KeyError, TypeError):
+                data['Driver'] = self.DATA_MISSING
 
         # Taskbar visibility
-        data['Taskbar'] = sync_status.get('taskbar', self.DATA_MISSING)
+        if SyncStatusRequestFlags.Taskbar in request_flags:
+            data['Taskbar'] = sync_status.get('taskbar', self.DATA_MISSING)
 
         # Operating system version
         data['OsVer'] = self.friendly_osver(device)
 
         # CPU utilization
-        try:
-            num_cores = len(sync_status['cpuUtilization'])
-            num_overloaded_cores = 0
-            cpu_load_avg = 0.0
-            for core_load in sync_status['cpuUtilization']:
-                cpu_load_avg += float(core_load) * (1.0 / num_cores)
-                if core_load > self.CORE_OVERLOAD_THRESH:
-                    num_overloaded_cores += 1
-
-            data['CpuUtilization'] = f"{cpu_load_avg:.0f}%"
-            if num_overloaded_cores > 0:
-                data['CpuUtilization'] += f' ({num_overloaded_cores} cores >' \
-                    f' {self.CORE_OVERLOAD_THRESH}%)'
-
-            if device.processor_smt:
-                data['CpuUtilization'] += ' (SMT ENABLED)'
-        except (KeyError, ValueError):
-            data['CpuUtilization'] = self.DATA_MISSING
+        if SyncStatusRequestFlags.CpuUtilization in request_flags:
+            try:
+                num_cores = len(sync_status['cpuUtilization'])
+                num_overloaded_cores = 0
+                cpu_load_avg = 0.0
+                for core_load in sync_status['cpuUtilization']:
+                    cpu_load_avg += float(core_load) * (1.0 / num_cores)
+                    if core_load > self.CORE_OVERLOAD_THRESH:
+                        num_overloaded_cores += 1
+
+                data['CpuUtilization'] = f"{cpu_load_avg:.0f}%"
+                if num_overloaded_cores > 0:
+                    data['CpuUtilization'] += f' ({num_overloaded_cores} cores >' \
+                        f' {self.CORE_OVERLOAD_THRESH}%)'
+
+                if device.processor_smt:
+                    data['CpuUtilization'] += ' (SMT ENABLED)'
+            except (KeyError, ValueError):
+                data['CpuUtilization'] = self.DATA_MISSING
 
         # Memory utilization
-        try:
-            gb = 1024 * 1024 * 1024
-            mem_total = device.total_phys_mem
-            mem_avail = sync_status.get('availablePhysicalMemory', 0)
-            mem_utilized = mem_total - mem_avail
-            data['MemUtilization'] = \
-                f'{mem_utilized/gb:.1f} / {mem_total/gb:.0f} GB'
-        except TypeError:
-            data['MemUtilization'] = self.DATA_MISSING
+        if SyncStatusRequestFlags.AvailablePhysicalMemory in request_flags:
+            try:
+                gb = 1024 * 1024 * 1024
+                mem_total = device.total_phys_mem
+                mem_avail = sync_status.get('availablePhysicalMemory', 0)
+                mem_utilized = mem_total - mem_avail
+                data['MemUtilization'] = \
+                    f'{mem_utilized/gb:.1f} / {mem_total/gb:.0f} GB'
+            except TypeError:
+                data['MemUtilization'] = self.DATA_MISSING
 
         # GPU utilization + clocks
-        try:
-            gpu_stats = list(map(
-                lambda x: f"#{x[0]}: {x[1]:.0f}% ({x[2] / 1000:.0f} MHz)",
-                zip(count(), sync_status['gpuUtilization'],
-                    sync_status['gpuCoreClocksKhz'])))
-
-            if len(gpu_stats) > 0:
-                data['GpuUtilization'] = '\n'.join(gpu_stats)
-            else:
+        if (SyncStatusRequestFlags.GpuUtilization | SyncStatusRequestFlags.GpuCoreClockKhz) in request_flags:
+            try:
+                gpu_stats = list(map(
+                    lambda x: f"#{x[0]}: {x[1]:.0f}% ({x[2] / 1000:.0f} MHz)",
+                    zip(count(), sync_status['gpuUtilization'],
+                        sync_status['gpuCoreClocksKhz'])))
+
+                if len(gpu_stats) > 0:
+                    data['GpuUtilization'] = '\n'.join(gpu_stats)
+                else:
+                    data['GpuUtilization'] = self.DATA_MISSING
+            except (KeyError, TypeError):
                 data['GpuUtilization'] = self.DATA_MISSING
-        except (KeyError, TypeError):
-            data['GpuUtilization'] = self.DATA_MISSING
 
         # GPU temperature
-        try:
-            temps = [t if t != -2147483648 else self.DATA_MISSING
-                     for t in sync_status['gpuTemperature']]
-
-            if len(temps) > 0:
-                data['GpuTemperature'] = '\n'.join(
-                    map(lambda x: f"#{x[0]}: {x[1]}° C", zip(count(), temps)))
-            else:
+        if SyncStatusRequestFlags.GpuTemperature in request_flags:
+            try:
+                temps = [t if t != -2147483648 else self.DATA_MISSING for t in sync_status['gpuTemperature']]
+
+                if len(temps) > 0:
+                    data['GpuTemperature'] = '\n'.join(
+                        map(lambda x: f"#{x[0]}: {x[1]}° C", zip(count(), temps)))
+                else:
+                    data['GpuTemperature'] = self.DATA_MISSING
+            except (KeyError, TypeError):
                 data['GpuTemperature'] = self.DATA_MISSING
-        except (KeyError, TypeError):
-            data['GpuTemperature'] = self.DATA_MISSING
 
     def on_get_sync_status(self, device, message):
         '''
@@ -531,8 +606,7 @@ class nDisplayMonitor(QAbstractTableModel):
             return
 
         row = deviceIdx + 1
-        self.dataChanged.emit(self.createIndex(row, 1),
-                              self.createIndex(row, len(self.colnames)))
+        self.refresh_display_for_row(row)
 
     def try_issue_console_exec(self, exec_str, executor=''):
         ''' Issues a console exec to the cluster '''
@@ -549,6 +623,22 @@ class nDisplayMonitor(QAbstractTableModel):
 
         return False
 
+    def refresh_display_for_row(self, row: int):
+        ''' Refreshes the data displayed in the given row. Zero-based index. '''
+
+        self.refresh_display_for_cells(
+            self.createIndex(row, 0),
+            self.createIndex(row, len(self.colnames) - 1)
+        )
+
+    def refresh_display_for_cells(self, upper_left: QModelIndex, lower_right: QModelIndex):
+        ''' Refreshes the cells in the given range, inclusive '''
+
+        self.dataChanged.emit(upper_left, lower_right)
+
+        # This was necessary to reliably update the table display with the changed data.
+        self.headerDataChanged.emit(Qt.Orientation.Horizontal, upper_left.column, lower_right.column)
+
     # ~ QAbstractTableModel interface begin
 
     def rowCount(self, parent=QModelIndex()):
@@ -560,6 +650,9 @@ class nDisplayMonitor(QAbstractTableModel):
     def headerData(self, section, orientation, role):
         if role == Qt.ItemDataRole.DisplayRole:
             if orientation == Qt.Orientation.Horizontal:
+                # Connection column uses icons and don't need the long header name
+                if self.colnames[section] == 'Connected':
+                    return ""
                 return self.colnames[section]
             else:
                 return "{}".format(section)
@@ -578,7 +671,20 @@ class nDisplayMonitor(QAbstractTableModel):
 
         # grab (device_hash, device_data) from ordered dict
         _, devicedata = list(self.devicedatas.items())[row]
+        device = devicedata['device']
         data = devicedata['data']
+        is_program_running = self.program_id_from_device(device) != self.default_program_id()
+
+        if colname == 'Connected':
+            if role == Qt.ItemDataRole.DecorationRole:
+                if data['Connected'].lower() == 'yes':
+                    if is_program_running:
+                        return self.icon_running
+                    return self.icon_connected
+                return self.icon_unconnected
+
+            return None
+
         value = data[colname]
 
         if role == Qt.ItemDataRole.DisplayRole:
@@ -586,22 +692,53 @@ class nDisplayMonitor(QAbstractTableModel):
 
         elif role == Qt.ItemDataRole.BackgroundRole:
             return self.color_for_column(colname=colname, value=value,
-                                         data=data)
+                                         data=data, is_program_running=is_program_running)
 
         elif role == Qt.ItemDataRole.TextAlignmentRole:
+            alignment = Qt.AlignmentFlag.AlignCenter
             if colname in ('CpuUtilization', 'GpuUtilization'):
-                return Qt.AlignmentFlag.AlignLeft
-            return Qt.AlignmentFlag.AlignRight
+                alignment = Qt.AlignmentFlag.AlignLeft | Qt.AlignmentFlag.AlignVCenter
+            return alignment
+
+        elif role == Qt.ItemDataRole.ForegroundRole:
+            if data['Connected'].lower() != 'yes':
+                return self.COLOR_DISCONNECTED
 
         return None
 
     # ~ QAbstractTableModel interface end
 
+    def update_device_name_and_address(self, in_device: Device):
+        ''' Refeshes the node names and addresses displayed in the table '''
+
+        for deviceIdx, devicedata in enumerate(self.devicedatas.values()):
+            device = devicedata['device']
+
+            # See if it is the right device based on its hash.
+            if in_device.device_hash == device.device_hash:
+
+                # Update name and address
+                data = devicedata['data']
+                data['Host'] = str(device.address)
+                data['Node'] = device.name
+
+                # Refresh the table with the new data
+                row = deviceIdx + 1
+                self.refresh_display_for_row(row)
+                return
+
     @QtCore.Slot()
-    def on_refresh_mosaics_clicked(self):
+    def on_refresh_table_clicked(self):
+        ''' Refreshes the data in the table, which involves sending commands to
+        the devices.'''
+
         for devicedata in self.devicedatas.values():
             device = devicedata['device']
             device.refresh_mosaics()
+            self.update_device_name_and_address(device)
+
+        # Request a full update on the sync status.
+        self.poll_sync_status(SyncStatusRequestFlags.all())
 
     @QtCore.Slot()
     def on_fix_exe_flags_clicked(self):
diff --git a/Engine/Plugins/VirtualProduction/Switchboard/Source/Switchboard/switchboard/devices/ndisplay/ndisplay_monitor_ui.py b/Engine/Plugins/VirtualProduction/Switchboard/Source/Switchboard/switchboard/devices/ndisplay/ndisplay_monitor_ui.py
--- a/Engine/Plugins/VirtualProduction/Switchboard/Source/Switchboard/switchboard/devices/ndisplay/ndisplay_monitor_ui.py
+++ b/Engine/Plugins/VirtualProduction/Switchboard/Source/Switchboard/switchboard/devices/ndisplay/ndisplay_monitor_ui.py
@@ -62,11 +62,11 @@ class nDisplayMonitorUI(QWidget):
             lambda: self.try_issue_console_exec())
 
         # Create remaining buttons
-        btn_refresh_mosaics = QPushButton('Refresh Mosaics')
-        btn_refresh_mosaics.setToolTip(
+        btn_refresh_table = QPushButton('Refresh Table')
+        btn_refresh_table.setToolTip(
             'Updates the cached mosaic topologies.')
-        btn_refresh_mosaics.clicked.connect(
-            self.monitor.on_refresh_mosaics_clicked)
+        btn_refresh_table.clicked.connect(
+            self.monitor.on_refresh_table_clicked)
 
         btn_fix_exe_flags = QPushButton('Fix ExeFlags')
         btn_fix_exe_flags.setToolTip(
@@ -84,8 +84,11 @@ class nDisplayMonitorUI(QWidget):
         layout_buttons.addWidget(self.cmb_console_exec)
         layout_buttons.addWidget(btn_console_exec)
         layout_buttons.addStretch(1)
-        layout_buttons.addWidget(btn_refresh_mosaics)
-        layout_buttons.addWidget(btn_fix_exe_flags)
+        layout_buttons.addWidget(btn_refresh_table)
+
+        if self.monitor.use_exe_flags:
+            layout_buttons.addWidget(btn_fix_exe_flags)
+
         layout_buttons.addWidget(btn_minimize_windows)
 
         return layout_buttons
diff --git a/Engine/Plugins/VirtualProduction/Switchboard/Source/Switchboard/switchboard/devices/ndisplay/plugin_ndisplay.py b/Engine/Plugins/VirtualProduction/Switchboard/Source/Switchboard/switchboard/devices/ndisplay/plugin_ndisplay.py
--- a/Engine/Plugins/VirtualProduction/Switchboard/Source/Switchboard/switchboard/devices/ndisplay/plugin_ndisplay.py
+++ b/Engine/Plugins/VirtualProduction/Switchboard/Source/Switchboard/switchboard/devices/ndisplay/plugin_ndisplay.py
@@ -25,6 +25,7 @@ from switchboard.devices.unreal.plugin_unreal import DeviceUnreal, \
     DeviceWidgetUnreal, LiveLinkPresetSetting, MediaProfileSetting
 from switchboard.devices.unreal.uassetparser import UassetParser
 from switchboard.devices.device_base import DeviceStatus
+from switchboard.message_protocol import SyncStatusRequestFlags
 from switchboard.switchboard_logging import LOGGER
 from switchboard.sbcache import SBCache, Asset
 
@@ -530,12 +531,12 @@ class DevicenDisplay(DeviceUnreal):
         'lock_gpu_clock': BoolSetting(
             attr_name="lock_gpu_clock",
             nice_name="Lock GPU Clock",
-            value=False,
+            value=True,  # nDisplay hitches less whe GPU clocks are locked.
             tool_tip=(
                 "Hint to lock the GPU clock to its allowed maximum. Requires SwitchboardListenerHelper \n"
-                "to be running on the client machine, otherwise this option will be ignored."
+                "to be running on the client machine as administrator, otherwise this option will be ignored."
             ),
-            show_ui = True if sys.platform in ('win32','linux') else False, # Gpu Clocker is available in select platforms
+            show_ui=True if sys.platform in ('win32', 'linux') else False,  # Gpu Clocker is available in select platforms
         ),
     }
 
@@ -598,7 +599,6 @@ class DevicenDisplay(DeviceUnreal):
         self.path_to_config_on_host = DevicenDisplay.csettings[
             'ndisplay_config_file'].get_value()
 
-
         CONFIG.ENGINE_DIR.signal_setting_overridden.connect(
             self.on_cmdline_affecting_override)
         CONFIG.ENGINE_DIR.signal_setting_changed.connect(
@@ -800,7 +800,7 @@ class DevicenDisplay(DeviceUnreal):
 
         if self.should_use_project_path_in_command_line():
             uproject = CONFIG.UPROJECT_PATH.get_value(self.name)
-        
+
         # normalize path if not empty
         if uproject != "":
             uproject = os.path.normpath(uproject)
@@ -983,13 +983,12 @@ class DevicenDisplay(DeviceUnreal):
                 '-dc_replicationserver_type listen'
             ])
 
-
         # fill in ExecCmds
         exec_cmds = self.csettings["ndisplay_exec_cmds"].get_value(self.name).copy()
 
         if DevicenDisplay.csettings['disable_all_screen_messages'].get_value() and 'DisableAllScreenMessages' not in exec_cmds:
             exec_cmds.append('DisableAllScreenMessages')
-        
+
         # LiveLink preset
         livelink_preset_gamepath = DevicenDisplay.csettings["livelink_preset"].get_value(self.name)
         if livelink_preset_gamepath:
@@ -1095,7 +1094,7 @@ class DevicenDisplay(DeviceUnreal):
 
         path_to_exe = self.generate_unreal_exe_path()
         args_expanded = ' '.join(args)
-        
+
         self.settings['ue_command_line'].update_value(
             f"{path_to_exe} {args_expanded}")
 
@@ -1170,19 +1169,6 @@ class DevicenDisplay(DeviceUnreal):
             LOGGER.error(
                 f'Error parsing "refresh mosaics" response ({message})')
 
-    def device_widget_registered(self, device_widget):
-        ''' Device interface method '''
-
-        super().device_widget_registered(device_widget)
-
-        # hook to its primary_button clicked event
-        device_widget.signal_device_widget_primary.connect(
-            self.select_as_primary)
-
-        # Update the state of the button depending on whether this device is
-        # the primary or not.
-        device_widget.primary_button.setChecked(self.is_primary())
-
     def get_primary_node_device(self):
         for device in self.active_unreal_devices:
             if device.device_type == 'nDisplay' and device.is_primary():
@@ -1398,9 +1384,8 @@ class DevicenDisplay(DeviceUnreal):
             return cls.parse_config_json(cfg_file)
         elif ext == '.uasset':
             return cls.parse_config_uasset(cfg_file)
-        
-        raise Exception(f'Unknown config extension "{ext}"')
 
+        raise Exception(f'Unknown config extension "{ext}"')
 
     def update_settings_controlled_by_config(self, cfg_file):
         '''
@@ -1682,6 +1667,7 @@ class DevicenDisplay(DeviceUnreal):
         except Exception:
             LOGGER.warning("Could not soft kill the cluster")
 
+    # ~ Begin Device Interface
     @classmethod
     def all_devices_added(cls):
         ''' Device interface implementation '''
@@ -1729,3 +1715,36 @@ class DevicenDisplay(DeviceUnreal):
             if not primary_found:
                 LOGGER.error('Could not find or assign primary device in the '
                              f'current devices and config file "{cfg_file}"')
+
+    def device_widget_registered(self, device_widget):
+        ''' Device interface method '''
+
+        super().device_widget_registered(device_widget)
+
+        # hook to its primary_button clicked event
+        device_widget.signal_device_widget_primary.connect(
+            self.select_as_primary)
+
+        # Update the state of the button depending on whether this device is
+        # the primary or not.
+        device_widget.primary_button.setChecked(self.is_primary())
+
+    def on_device_name_changed(self, new_name):
+        super().on_device_name_changed(new_name)
+        if self.__class__.ndisplay_monitor:
+            self.__class__.ndisplay_monitor.update_device_name_and_address(self)
+
+    def on_address_changed(self, new_address):
+        super().on_address_changed(new_address)
+
+        if self.__class__.ndisplay_monitor:
+            self.__class__.ndisplay_monitor.update_device_name_and_address(self)
+    # ~ End Device Interface
+
+    # ~ Begin DeviceUnreal Interface
+    @QtCore.Slot()
+    def _on_listener_connected(self):
+        if self.__class__.ndisplay_monitor:
+            # Poll the sync status to get an immediate update on its state
+            self.__class__.ndisplay_monitor.poll_sync_status_for_device(self, SyncStatusRequestFlags.all())
+    # ~ End DeviceUnrealInterface
diff --git a/Engine/Plugins/VirtualProduction/Switchboard/Source/Switchboard/switchboard/devices/unreal/version_helpers.py b/Engine/Plugins/VirtualProduction/Switchboard/Source/Switchboard/switchboard/devices/unreal/version_helpers.py
--- a/Engine/Plugins/VirtualProduction/Switchboard/Source/Switchboard/switchboard/devices/unreal/version_helpers.py
+++ b/Engine/Plugins/VirtualProduction/Switchboard/Source/Switchboard/switchboard/devices/unreal/version_helpers.py
@@ -3,7 +3,7 @@
 import typing
 
 
-LISTENER_COMPATIBLE_VERSION = (3,0)
+LISTENER_COMPATIBLE_VERSION = (3,1)
 
 # Version where redeploy support was initially introduced.
 LISTENER_MIN_REDEPLOY_VERSION = (1,4,0)
diff --git a/Engine/Plugins/VirtualProduction/Switchboard/Source/Switchboard/switchboard/message_protocol.py b/Engine/Plugins/VirtualProduction/Switchboard/Source/Switchboard/switchboard/message_protocol.py
--- a/Engine/Plugins/VirtualProduction/Switchboard/Source/Switchboard/switchboard/message_protocol.py
+++ b/Engine/Plugins/VirtualProduction/Switchboard/Source/Switchboard/switchboard/message_protocol.py
@@ -4,6 +4,38 @@ import base64
 import json
 from typing import Optional
 import uuid
+from enum import IntFlag
+from functools import reduce
+
+
+class SyncStatusRequestFlags(IntFlag):
+    ''' Bit flag to represent the sync status being requested.
+    We support this granularity because certain sync status requests
+    can be heavier on the system than others.
+    It should match its counterpart's definition in SwitchboardTasks.h.
+    '''
+    Nothing = 0
+    SyncTopos = 1 << 0
+    MosaicTopos = 1 << 1
+    FlipModeHistory = 1 << 2
+    ProgramLayers = 1 << 3
+    DriverInfo = 1 << 4
+    Taskbar = 1 << 5
+    PidInFocus = 1 << 6
+    CpuUtilization = 1 << 7
+    AvailablePhysicalMemory = 1 << 8
+    GpuUtilization = 1 << 9
+    GpuCoreClockKhz = 1 << 10
+    GpuTemperature = 1 << 11
+
+    def __contains__(self, flags):
+        ''' Returns true if the given bit flag is set'''
+        return (self & flags) == flags
+
+    @classmethod
+    def all(cls):
+        ''' Returns in intance of this enum with all flags set'''
+        return reduce(lambda x, y: x | y, cls)
 
 
 def create_authenticate_message(
@@ -62,73 +94,97 @@ def create_start_process_message(
     message = json.dumps(start_cmd).encode() + b'\x00'
     return (cmd_id, message)
 
+
 def create_kill_process_message(program_id):
     cmd_id = uuid.uuid4()
     kill_cmd = {'command': 'kill', 'id': str(cmd_id), 'uuid': str(program_id)}
     message = json.dumps(kill_cmd).encode() + b'\x00'
     return (cmd_id, message)
 
+
 def create_vcs_init_message(provider, vcs_settings):
     cmd_id = uuid.uuid4()
     vcs_init_cmd = {'command': 'vcs init', 'id': str(cmd_id), 'provider': provider, 'vcs settings': vcs_settings}
     message = json.dumps(vcs_init_cmd).encode() + b'\x00'
     return (cmd_id, message)
 
+
 def create_vcs_report_revision_message(path):
     cmd_id = uuid.uuid4()
     vcs_revision_cmd = {'command': 'vcs report revision', 'id': str(cmd_id), 'path': path}
     message = json.dumps(vcs_revision_cmd).encode() + b'\x00'
     return (cmd_id, message)
 
+
 def create_vcs_sync_message(revision, path):
     cmd_id = uuid.uuid4()
     vcs_sync_cmd = {'command': 'vcs sync', 'id': str(cmd_id), 'revision': revision, 'path': path}
     message = json.dumps(vcs_sync_cmd).encode() + b'\x00'
     return (cmd_id, message)
 
+
 def create_disconnect_message():
     cmd_id = uuid.uuid4()
     disconnect_cmd = {'command': 'disconnect', 'id': str(cmd_id)}
     message = json.dumps(disconnect_cmd).encode() + b'\x00'
     return (cmd_id, message)
 
+
 def create_send_filecontent_message(file_content: bytes, destination_path: str, *, force_overwrite: bool = False):
     encoded_content = base64.b64encode(file_content)
 
     cmd_id = uuid.uuid4()
-    transfer_file_cmd = {'command': 'send file', 'id': str(cmd_id), 'destination': destination_path, 'content': encoded_content.decode(), 'force_overwrite': force_overwrite}
+    transfer_file_cmd = {
+        'command': 'send file',
+        'id': str(cmd_id),
+        'destination': destination_path,
+        'content': encoded_content.decode(),
+        'force_overwrite': force_overwrite
+    }
     message = json.dumps(transfer_file_cmd).encode() + b'\x00'
     return (cmd_id, message)
 
+
 def create_send_file_message(path_to_source_file: str, destination_path: str, *, force_overwrite: bool = False):
     with open(path_to_source_file, 'rb') as f:
         file_content = f.read()
     return create_send_filecontent_message(file_content, destination_path, force_overwrite=force_overwrite)
 
+
 def create_copy_file_from_listener_message(path_on_listener_machine):
     cmd_id = uuid.uuid4()
     copy_file_cmd = {'command': 'receive file', 'id': str(cmd_id), 'source': path_on_listener_machine}
     message = json.dumps(copy_file_cmd).encode() + b'\x00'
     return (cmd_id, message)
 
+
 def create_keep_alive_message():
     cmd_id = uuid.uuid4()
     keep_alive_cmd = {'command': 'keep alive', 'id': str(cmd_id)}
     message = json.dumps(keep_alive_cmd).encode() + b'\x00'
     return (cmd_id, message)
 
-def create_get_sync_status_message(program_id):
+
+def create_get_sync_status_message(program_id, request_flags: SyncStatusRequestFlags):
     cmd_id = uuid.uuid4()
-    cmd = {'command': 'get sync status', 'id': str(cmd_id), 'uuid': str(program_id), 'bEcho': False}
+    cmd = {
+        'command': 'get sync status',
+        'id': str(cmd_id),
+        'uuid': str(program_id),
+        'request_flags': int(request_flags),
+        'bEcho': False
+    }
     message = json.dumps(cmd).encode() + b'\x00'
     return (cmd_id, message)
 
+
 def create_refresh_mosaics_message():
     cmd_id = uuid.uuid4()
     cmd = {'command': 'refresh mosaics', 'id': str(cmd_id)}
     message = json.dumps(cmd).encode() + b'\x00'
     return (cmd_id, message)
 
+
 def create_redeploy_listener_message(base64listener: str, sha1digest: str):
     ''' Sends a command to replace the remote server's listener executable. '''
     cmd_id = uuid.uuid4()
@@ -136,6 +192,7 @@ def create_redeploy_listener_message(base64listener: str, sha1digest: str):
     message = json.dumps(redeploy_cmd).encode() + b'\x00'
     return (cmd_id, message)
 
+
 def create_free_listener_bin_message():
     ''' Sends a command to the listener to move its executable. '''
     cmd_id = uuid.uuid4()
@@ -143,23 +200,27 @@ def create_free_listener_bin_message():
     message = json.dumps(rename_proc_cmd).encode() + b'\x00'
     return (cmd_id, message)
 
+
 def create_fixExeFlags_message(puuid):
     cmd_id = uuid.uuid4()
     cmd = {'command': 'fixExeFlags', 'id': str(cmd_id), 'uuid': str(puuid)}
     message = json.dumps(cmd).encode() + b'\x00'
     return (cmd_id, message)
 
+
 def decode_message(msg_as_bytes: bytes):
     msg_as_str = msg_as_bytes.decode()
     msg_json = json.loads(msg_as_str)
     return msg_json
 
+
 def create_minimize_windows_message():
     cmd_id = uuid.uuid4()
     cmd = {'command': 'minimize windows', 'id': str(cmd_id)}
     message = json.dumps(cmd).encode() + b'\x00'
     return (cmd_id, message)
 
+
 def create_set_inactive_timeout_message(timeout_seconds: int):
     cmd_id = uuid.uuid4()
     cmd = {'command': 'set inactive timeout', 'id': str(cmd_id),
diff --git a/Engine/Source/Programs/SwitchboardListener/SblCore/Internal/SwitchboardListener.h b/Engine/Source/Programs/SwitchboardListener/SblCore/Internal/SwitchboardListener.h
--- a/Engine/Source/Programs/SwitchboardListener/SblCore/Internal/SwitchboardListener.h
+++ b/Engine/Source/Programs/SwitchboardListener/SblCore/Internal/SwitchboardListener.h
@@ -267,7 +267,6 @@ private:
 	/** Client interface to the Switchboard Listener Helper external process */
 	TSharedPtr<FSBLHelperClient> SBLHelper;
 
-	bool bIsNvAPIInitialized;
 	bool bProcessorSMT;
 
 	TSharedPtr<FRWLock> CachedMosaicToposLock;
diff --git a/Engine/Source/Programs/SwitchboardListener/SblCore/Internal/SwitchboardListenerVersion.h b/Engine/Source/Programs/SwitchboardListener/SblCore/Internal/SwitchboardListenerVersion.h
--- a/Engine/Source/Programs/SwitchboardListener/SblCore/Internal/SwitchboardListenerVersion.h
+++ b/Engine/Source/Programs/SwitchboardListener/SblCore/Internal/SwitchboardListenerVersion.h
@@ -3,5 +3,5 @@
 #pragma once
 
 #define SBLISTENER_VERSION_MAJOR 3
-#define SBLISTENER_VERSION_MINOR 0
-#define SBLISTENER_VERSION_PATCH 7
+#define SBLISTENER_VERSION_MINOR 1
+#define SBLISTENER_VERSION_PATCH 0
diff --git a/Engine/Source/Programs/SwitchboardListener/SblCore/NvmlWrapper/Public/GpuClocker.h b/Engine/Source/Programs/SwitchboardListener/SblCore/NvmlWrapper/Public/GpuClocker.h
new file mode 100644
--- /dev/null
+++ b/Engine/Source/Programs/SwitchboardListener/SblCore/NvmlWrapper/Public/GpuClocker.h
@@ -0,0 +1,38 @@
+// Copyright Epic Games, Inc. All Rights Reserved.
+
+#pragma once
+
+#include <stdint.h>
+
+#include "NvmlWrapperPublic.h"
+
+/** Encapsulates the logic to lock/unlock gpu clocks */
+class EXPORTLIB FGpuClocker
+{
+public:
+
+	/** Locks the gpu clocks to the maximum allowed */
+	bool LockGpuClocks() const;
+
+	/** Unlocks the gpu and memory clocks */
+	bool UnlockGpuClocks() const;
+
+	/** Locks the given gpu and memory clocks to the maximum allowed */
+	bool GpuUnlockDevice(const uint32_t GpuIdx) const;
+
+	/** Unlocks the given gpu and memory clocks */
+	bool GpuLockDevice(const uint32_t GpuIdx) const;
+
+	/** Returns number of GPUs */
+	bool GetGpuCount(uint32_t& OutNumGpus) const;
+
+	/** Returns the temperature in Celsius of the indexed GPU */
+	bool GetGpuCelsius(const uint32_t GpuIdx, uint32_t& OutCelsius) const;
+
+	/** Returns the GPU usage in percentage of running kernels and memory operations */
+	bool GetGpuUsage(const uint32_t GpuIdx, uint32_t& OutUsageKernel, uint32_t& OutUsageMemory) const;
+
+	/** Return GPU clock frequency in MHz */
+	bool GetGpuMHz(const uint32_t GpuIdx, uint32_t& GraphicsMHz, uint32_t& MemoryMHz) const;
+};
+
diff --git a/Engine/Source/Programs/SwitchboardListener/SblCore/NvmlWrapper/Public/NvmlWrapperPublic.h b/Engine/Source/Programs/SwitchboardListener/SblCore/NvmlWrapper/Public/NvmlWrapperPublic.h
new file mode 100644
--- /dev/null
+++ b/Engine/Source/Programs/SwitchboardListener/SblCore/NvmlWrapper/Public/NvmlWrapperPublic.h
@@ -0,0 +1,14 @@
+// Copyright Epic Games, Inc. All Rights Reserved.
+
+#pragma once
+
+#if defined(_MSC_VER)
+	#define EXPORTLIB __declspec(dllexport)
+	#define IMPORTLIB __declspec(dllimport)
+#elif defined(__GNUC__)
+	#define EXPORTLIB __attribute__((visibility("default")))
+	#define IMPORTLIB
+#else
+	#define EXPORTLIB
+	#define IMPORTLIB
+#endif
diff --git a/Engine/Source/Programs/SwitchboardListener/SblCore/Private/SwitchboardListener.cpp b/Engine/Source/Programs/SwitchboardListener/SblCore/Private/SwitchboardListener.cpp
--- a/Engine/Source/Programs/SwitchboardListener/SblCore/Private/SwitchboardListener.cpp
+++ b/Engine/Source/Programs/SwitchboardListener/SblCore/Private/SwitchboardListener.cpp
@@ -33,6 +33,8 @@
 
 #if PLATFORM_WINDOWS
 
+#include "GpuClocker.h"
+
 #pragma warning(push)
 #pragma warning(disable : 4005)	// Disable macro redefinition warning for compatibility with Windows SDK 8+
 
@@ -253,31 +255,109 @@ FString FSwitchboardCommandLineOptions::ToString(bool bIncludeRedeploy /* = fals
 	return FString::Join(Args, TEXT(" "));
 }
 
-FSwitchboardListener::FSwitchboardListener(const FSwitchboardCommandLineOptions& InOptions)
-	: Options(InOptions)
-	, CpuMonitor(MakeShared<FCpuUtilizationMonitor>())
-	, SBLHelper(MakeShared<FSBLHelperClient>())
-	, bIsNvAPIInitialized(false)
-	, CachedMosaicToposLock(MakeShared<FRWLock>())
-	, CachedMosaicTopos(MakeShared<TArray<FMosaicTopo>>())
-{
+
 #if PLATFORM_WINDOWS
-	// initialize NvAPI
+/** Used to limit NvApi loading to the current scope (recursive allowed) and thread. */
+class FScopedNvApi
+{
+private:
+
+	/** Global recursive lock for NvApi usage */
+	static UE::FRecursiveMutex NvApiLock;
+
+	/** Keeps track of failed unloads. Useful to recover from potential NvApi unloading errors */
+	static std::atomic<int32> OwedUnloads;
+
+	/** Used to repeatedly show NvAPI initialization errors. Once is enough. */
+	static bool bErrorLoggedOnce;
+
+	/** True if the NvAPI was initialized */
+	bool bIsNvApiInitialized;
+
+public:
+	FScopedNvApi()
+		: bIsNvApiInitialized(false)
 	{
-		const NvAPI_Status Result = NvAPI_Initialize();
-		if (Result == NVAPI_OK)
-		{
-			bIsNvAPIInitialized = true;
+		NvApiLock.Lock(); // Unlocked at the end of the destructor
 
-			FillOutMosaicTopologies(*CachedMosaicTopos);
+		// Initialize NvAPI.
+
+		if (OwedUnloads)
+		{
+			OwedUnloads--;
+			bIsNvApiInitialized = true; // If it isn't true, the calls to the api will simply fail.
 		}
 		else
 		{
-			NvAPI_ShortString ErrorString;
-			NvAPI_GetErrorMessage(Result, ErrorString);
-			UE_LOG(LogSwitchboard, Error, TEXT("NvAPI_Initialize failed. Error: %s"), ANSI_TO_TCHAR(ErrorString));
+			const NvAPI_Status Result = NvAPI_Initialize();
+
+			if (Result == NVAPI_OK)
+			{
+				bIsNvApiInitialized = true;
+			}
+			else if (!bErrorLoggedOnce)
+			{
+				bErrorLoggedOnce = false;
+
+				NvAPI_ShortString ErrorString;
+				NvAPI_GetErrorMessage(Result, ErrorString);
+				UE_LOG(LogSwitchboard, Error, TEXT("NvAPI_Initialize failed. Error: %s"), ANSI_TO_TCHAR(ErrorString));
+			}
+		}
+	}
+
+	~FScopedNvApi()
+	{
+		// Only unload NvApi if it was initialized
+		if (bIsNvApiInitialized)
+		{
+			const NvAPI_Status Result = NvAPI_Unload();
+
+			if (Result == NVAPI_OK)
+			{
+				bIsNvApiInitialized = false;
+			}
+			else
+			{
+				// This would be an unexpected event because the lock prevents concurrent usage of the library.
+				// We try to compensate by not loading next time.
+
+				OwedUnloads++;
+
+				NvAPI_ShortString ErrorString;
+				NvAPI_GetErrorMessage(Result, ErrorString);
+				UE_LOG(LogSwitchboard, Error, TEXT("NvAPI_Unload unexpectedly failed. Error: %s"), ANSI_TO_TCHAR(ErrorString));
+			}
 		}
+
+		NvApiLock.Unlock();
 	}
+
+	/** Returns true if the NvAPI was initialized successfully */
+	bool IsNvApiInitialized()
+	{
+		return bIsNvApiInitialized;
+	}
+};
+
+// Initialize static member variables
+bool FScopedNvApi::bErrorLoggedOnce = false;
+std::atomic<int32> FScopedNvApi::OwedUnloads = 0;
+UE::FRecursiveMutex FScopedNvApi::NvApiLock;
+
+#endif // PLATFORM_WINDOWS
+
+
+FSwitchboardListener::FSwitchboardListener(const FSwitchboardCommandLineOptions& InOptions)
+	: Options(InOptions)
+	, CpuMonitor(MakeShared<FCpuUtilizationMonitor>())
+	, SBLHelper(MakeShared<FSBLHelperClient>())
+	, CachedMosaicToposLock(MakeShared<FRWLock>())
+	, CachedMosaicTopos(MakeShared<TArray<FMosaicTopo>>())
+{
+#if PLATFORM_WINDOWS
+	// Cache Mosaic Topologies
+	FillOutMosaicTopologies(*CachedMosaicTopos);
 #endif // PLATFORM_WINDOWS
 
 	const int32 NumLogicalProcessors = FPlatformMisc::NumberOfCoresIncludingHyperthreads();
@@ -984,6 +1064,12 @@ bool FSwitchboardListener::Task_StartProcess(const FSwitchboardStartTask& InRunT
 			// We disconnect right away because launches happen only far and in between.
 			SBLHelper->Disconnect();
 		}
+		else
+		{
+			UE_LOG(LogSwitchboard, Warning, TEXT("Lock Gpu clocks was requested but could not connect to SwitchboardListenerHelper process. "
+				"Please verify that it is running as admin (elevated privileges are required to lock Gpu clocks). "
+			    "If locking Gpu clocks is not desired, this option can be disabled in Switchboard."));
+		}
 	}
 
 	UE_LOG(LogSwitchboard, Display, TEXT("Started process %d: %s %s"), NewProcess->PID, *InRunTask.Command, *InRunTask.Arguments);
@@ -1780,16 +1866,17 @@ stream_receive_outer_break:
 }
 
 
-#if PLATFORM_WINDOWS
-static FCriticalSection SwitchboardListenerMutexNvapi;
-#endif // PLATFORM_WINDOWS
-
 #if PLATFORM_WINDOWS
 static void FillOutSyncTopologies(TArray<FSyncTopo>& SyncTopos)
 {
 	SWITCHBOARD_TRACE_CPUPROFILER_EVENT_SCOPE(FillOutSyncTopologies);
 
-	FScopeLock LockNvapi(&SwitchboardListenerMutexNvapi);
+	FScopedNvApi ScopedNvApi;
+
+	if (!ScopedNvApi.IsNvApiInitialized())
+	{
+		return;
+	}
 
 	// Normally there is a single sync card. BUT an RTX Server could have more, and we need to account for that.
 
@@ -1968,6 +2055,13 @@ static void FillOutDriverVersion(FSyncStatus& SyncStatus)
 {
 	SWITCHBOARD_TRACE_CPUPROFILER_EVENT_SCOPE(FillOutDriverVersion);
 
+	FScopedNvApi ScopedNvApi;
+
+	if (!ScopedNvApi.IsNvApiInitialized())
+	{
+		return;
+	}
+
 	NvU32 DriverVersion;
 	NvAPI_ShortString BuildBranchString;
 
@@ -2015,7 +2109,12 @@ static void FillOutMosaicTopologies(TArray<FMosaicTopo>& MosaicTopos)
 {
 	SWITCHBOARD_TRACE_CPUPROFILER_EVENT_SCOPE(FillOutMosaicTopologies);
 
-	FScopeLock LockNvapi(&SwitchboardListenerMutexNvapi);
+	FScopedNvApi ScopedNvApi;
+
+	if (!ScopedNvApi.IsNvApiInitialized())
+	{
+		return;
+	}
 
 	NvU32 GridCount = 0;
 	TArray<NV_MOSAIC_GRID_TOPO> GridTopologies;
@@ -2274,129 +2373,65 @@ static void FillOutDisableFullscreenOptimizationForProcess(FSyncStatus& SyncStat
 #endif // PLATFORM_WINDOWS
 
 #if PLATFORM_WINDOWS
-static void FillOutPhysicalGpuStats(FSyncStatus& SyncStatus)
+static void FillOutPhysicalGpuStats(FSyncStatus& SyncStatus, bool bGetUtilizations, bool bGetClocks, bool bGetTemperatures)
 {
 	SWITCHBOARD_TRACE_CPUPROFILER_EVENT_SCOPE(FillOutPhysicalGpuStats);
 
-	// TODO: Can we somehow use the GPU engine "engtype_3D" perf counters for this instead?
-
-	FScopeLock LockNvapi(&SwitchboardListenerMutexNvapi);
-
-	TArray<NvPhysicalGpuHandle> PhysicalGpuHandles;
-	PhysicalGpuHandles.SetNumUninitialized(NVAPI_MAX_PHYSICAL_GPUS);
-	NvU32 PhysicalGpuCount;
-	NvAPI_Status NvResult = NvAPI_EnumPhysicalGPUs(PhysicalGpuHandles.GetData(), &PhysicalGpuCount);
-	if (NvResult != NVAPI_OK)
+	if (!bGetUtilizations && !bGetClocks && !bGetTemperatures)
 	{
-		NvAPI_ShortString ErrorString;
-		NvAPI_GetErrorMessage(NvResult, ErrorString);
-		UE_LOG(LogSwitchboard, Warning, TEXT("NvAPI_EnumPhysicalGPUs failed. Error: %s"), ANSI_TO_TCHAR(ErrorString));
 		return;
 	}
 
-	PhysicalGpuHandles.SetNum(PhysicalGpuCount);
+	FGpuClocker GpuClocker;
 
-	// Sort first by bus, then by bus slot, ascending. Consistent with task manager and others.
-	Algo::Sort(PhysicalGpuHandles, [](const NvPhysicalGpuHandle& Lhs, const NvPhysicalGpuHandle& Rhs) -> bool {
-		NvU32 LhsBusId, RhsBusId, LhsSlotId, RhsSlotId;
-		const NvAPI_Status LhsBusResult = NvAPI_GPU_GetBusId(Lhs, &LhsBusId);
-		const NvAPI_Status RhsBusResult = NvAPI_GPU_GetBusId(Rhs, &RhsBusId);
-		const NvAPI_Status LhsSlotResult = NvAPI_GPU_GetBusSlotId(Lhs, &LhsSlotId);
-		const NvAPI_Status RhsSlotResult = NvAPI_GPU_GetBusSlotId(Rhs, &RhsSlotId);
-
-		if (LhsBusResult != NVAPI_OK || RhsBusResult != NVAPI_OK)
-		{
-			NvAPI_ShortString LhsErrorString;
-			NvAPI_ShortString RhsErrorString;
-			NvAPI_GetErrorMessage(LhsBusResult, LhsErrorString);
-			NvAPI_GetErrorMessage(RhsBusResult, RhsErrorString);
-			UE_LOG(LogSwitchboard, Warning, TEXT("NvAPI_GPU_GetBusId failed. Errors: %s, %s"), ANSI_TO_TCHAR(LhsErrorString), ANSI_TO_TCHAR(RhsErrorString));
-			return false;
-		}
+	uint32 GpuCount;
 
-		if (LhsSlotResult != NVAPI_OK || RhsSlotResult != NVAPI_OK)
-		{
-			NvAPI_ShortString LhsErrorString;
-			NvAPI_ShortString RhsErrorString;
-			NvAPI_GetErrorMessage(LhsSlotResult, LhsErrorString);
-			NvAPI_GetErrorMessage(RhsSlotResult, RhsErrorString);
-			UE_LOG(LogSwitchboard, Warning, TEXT("NvAPI_GPU_GetBusSlotId failed. Errors: %s, %s"), ANSI_TO_TCHAR(LhsErrorString), ANSI_TO_TCHAR(RhsErrorString));
-			return false;
-		}
-
-		if (LhsBusId != RhsBusId)
-		{
-			return LhsBusId < RhsBusId;
-		}
-
-		return LhsSlotId < RhsSlotId;
-	});
+	if (!GpuClocker.GetGpuCount(GpuCount))
+	{
+		return;
+	}
 
-	SyncStatus.GpuUtilization.SetNumUninitialized(PhysicalGpuCount);
-	SyncStatus.GpuCoreClocksKhz.SetNumUninitialized(PhysicalGpuCount);
-	SyncStatus.GpuTemperature.SetNumUninitialized(PhysicalGpuCount);
+	SyncStatus.GpuUtilization.SetNumUninitialized(GpuCount);
+	SyncStatus.GpuCoreClocksKhz.SetNumUninitialized(GpuCount);
+	SyncStatus.GpuTemperature.SetNumUninitialized(GpuCount);
 
-	for (NvU32 PhysicalGpuIdx = 0; PhysicalGpuIdx < PhysicalGpuCount; ++PhysicalGpuIdx)
+	for (uint32 GpuIdx = 0; GpuIdx < GpuCount; ++GpuIdx)
 	{
-		SyncStatus.GpuUtilization[PhysicalGpuIdx] = -1;
-		SyncStatus.GpuCoreClocksKhz[PhysicalGpuIdx] = -1;
-		SyncStatus.GpuTemperature[PhysicalGpuIdx] = MIN_int32;
+		SyncStatus.GpuUtilization[GpuIdx] = -1;
+		SyncStatus.GpuCoreClocksKhz[GpuIdx] = -1;
+		SyncStatus.GpuTemperature[GpuIdx] = MIN_int32;
 
-		const NvPhysicalGpuHandle& PhysicalGpu = PhysicalGpuHandles[PhysicalGpuIdx];
-
-		NV_GPU_DYNAMIC_PSTATES_INFO_EX PstatesInfo;
-		PstatesInfo.version = NV_GPU_DYNAMIC_PSTATES_INFO_EX_VER;
-		NvResult = NvAPI_GPU_GetDynamicPstatesInfoEx(PhysicalGpu, &PstatesInfo);
-		if (NvResult == NVAPI_OK)
+		if (bGetUtilizations)
 		{
-			// FIXME: NV_GPU_UTILIZATION_DOMAIN_ID enum is missing in our nvapi.h, but documented elsewhere.
-			//const int8 UtilizationDomain = NVAPI_GPU_UTILIZATION_DOMAIN_GPU;
-			const int8 UtilizationDomain = 0;
-			if (PstatesInfo.utilization[UtilizationDomain].bIsPresent)
+			uint32 UsageKernel;
+			uint32 UsageMemory;
+
+			if (GpuClocker.GetGpuUsage(GpuIdx, UsageKernel, UsageMemory))
 			{
-				SyncStatus.GpuUtilization[PhysicalGpuIdx] = PstatesInfo.utilization[UtilizationDomain].percentage;
+				SyncStatus.GpuUtilization[GpuIdx] = static_cast<int8>(UsageKernel);
 			}
 		}
-		else
-		{
-			UE_LOG(LogSwitchboard, Warning, TEXT("NvAPI_GPU_GetDynamicPstatesInfoEx failed. Error code: %d"), NvResult);
-		}
 
-		NV_GPU_CLOCK_FREQUENCIES ClockFreqs;
-		ClockFreqs.version = NV_GPU_CLOCK_FREQUENCIES_VER;
-		ClockFreqs.ClockType = NV_GPU_CLOCK_FREQUENCIES_CURRENT_FREQ;
-		NvResult = NvAPI_GPU_GetAllClockFrequencies(PhysicalGpu, &ClockFreqs);
-		if (NvResult == NVAPI_OK)
+		if (bGetClocks)
 		{
-			if (ClockFreqs.domain[NVAPI_GPU_PUBLIC_CLOCK_GRAPHICS].bIsPresent)
+			uint32 GraphicsMHz;
+			uint32 MemoryMHz;
+
+			if (GpuClocker.GetGpuMHz(GpuIdx, GraphicsMHz, MemoryMHz))
 			{
-				SyncStatus.GpuCoreClocksKhz[PhysicalGpuIdx] = ClockFreqs.domain[NVAPI_GPU_PUBLIC_CLOCK_GRAPHICS].frequency;
+				SyncStatus.GpuCoreClocksKhz[GpuIdx] = GraphicsMHz * 1000;
 			}
 		}
-		else
-		{
-			UE_LOG(LogSwitchboard, Warning, TEXT("NvAPI_GPU_GetAllClockFrequencies failed. Error code: %d"), NvResult);
-		}
 
-		NV_GPU_THERMAL_SETTINGS ThermalSettings;
-		ThermalSettings.version = NV_GPU_THERMAL_SETTINGS_VER;
-		NvResult = NvAPI_GPU_GetThermalSettings(PhysicalGpu, NVAPI_THERMAL_TARGET_ALL, &ThermalSettings);
-		if (NvResult == NVAPI_OK)
+		if (bGetTemperatures)
 		{
-			// Report max temp across all sensors for this GPU.
-			for (NvU32 SensorIdx = 0; SensorIdx < ThermalSettings.count; ++SensorIdx)
+			uint32 Celsius;
+
+			if (GpuClocker.GetGpuCelsius(GpuIdx, Celsius))
 			{
-				const NvS32 SensorTemp = ThermalSettings.sensor[SensorIdx].currentTemp;
-				if (SensorTemp > SyncStatus.GpuTemperature[PhysicalGpuIdx])
-				{
-					SyncStatus.GpuTemperature[PhysicalGpuIdx] = SensorTemp;
-				}
+				SyncStatus.GpuTemperature[GpuIdx] = static_cast<int32>(Celsius);
 			}
 		}
-		else
-		{
-			UE_LOG(LogSwitchboard, Warning, TEXT("NvAPI_GPU_GetThermalSettings failed. Error code: %d"), NvResult);
-		}
 	}
 }
 #endif
@@ -2428,9 +2463,13 @@ bool FSwitchboardListener::Task_GetSyncStatus(const FSwitchboardGetSyncStatusTas
 	TSharedRef<FSyncStatus> SyncStatus = MakeShared<FSyncStatus>(); // Smart pointer to avoid potentially bigger copy to lambda below.
 
 	// We need to run these on this thread to avoid threading issues.
-	FillOutFlipMode(SyncStatus.Get(), FindOrStartFlipModeMonitorForUUID(InGetSyncStatusTask.ProgramID));
+	if (EnumHasAnyFlags(InGetSyncStatusTask.RequestFlags, ESyncStatusRequestFlags::FlipModeHistory))
+	{
+		FillOutFlipMode(SyncStatus.Get(), FindOrStartFlipModeMonitorForUUID(InGetSyncStatusTask.ProgramID));
+	}
 
 	// Fill out fullscreen optimization setting
+	if (EnumHasAnyFlags(InGetSyncStatusTask.RequestFlags, ESyncStatusRequestFlags::ProgramLayers))
 	{
 		TSharedPtr<FRunningProcess>* ProcessPtr = RunningProcesses.FindByPredicate([&](const TSharedPtr<FRunningProcess>& InProcess)
 		{
@@ -2457,42 +2496,69 @@ bool FSwitchboardListener::Task_GetSyncStatus(const FSwitchboardGetSyncStatusTas
 	MessageFuture.Future = Async(EAsyncExecution::ThreadPool,
 		[
 			SyncStatus,
-			IsNvAPIInitialized=bIsNvAPIInitialized,
-			CpuMonitor=CpuMonitor,
-			CachedMosaicToposLock=CachedMosaicToposLock,
-			CachedMosaicTopos=CachedMosaicTopos
+			RequestFlags = InGetSyncStatusTask.RequestFlags,
+			CpuMonitor = CpuMonitor,
+			CachedMosaicToposLock = CachedMosaicToposLock,
+			CachedMosaicTopos = CachedMosaicTopos
 		]() {
 			SWITCHBOARD_TRACE_CPUPROFILER_EVENT_SCOPE_STR("FSwitchboardListener::Task_GetSyncStatus future closure");
 
-			FillOutTaskbarAutoHide(SyncStatus.Get());
+			// Load NvApi with FScopedNvApi only if we have requests that need it.
+			ESyncStatusRequestFlags FlagsThatNeedNvApi = ESyncStatusRequestFlags::DriverInfo;
+			EnumAddFlags(FlagsThatNeedNvApi, ESyncStatusRequestFlags::SyncTopos);
 
-			if (IsNvAPIInitialized)
+			if (EnumHasAnyFlags(RequestFlags, FlagsThatNeedNvApi))
 			{
-				FillOutDriverVersion(SyncStatus.Get());
-				FillOutSyncTopologies(SyncStatus->SyncTopos);
+				FScopedNvApi ScopedNvApi;
+
+				if (EnumHasAnyFlags(RequestFlags, ESyncStatusRequestFlags::DriverInfo))
+				{
+					FillOutDriverVersion(SyncStatus.Get());
+				}
+
+				if (EnumHasAnyFlags(RequestFlags, ESyncStatusRequestFlags::SyncTopos))
+				{
+					FillOutSyncTopologies(SyncStatus->SyncTopos);
+				}
 			}
 
+			if (EnumHasAnyFlags(RequestFlags, ESyncStatusRequestFlags::MosaicTopos))
 			{
 				FReadScopeLock Lock(*CachedMosaicToposLock);
 				SyncStatus->MosaicTopos = *CachedMosaicTopos;
 			}
 
-			SyncStatus->PidInFocus = FindPidInFocus();
+			if (EnumHasAnyFlags(RequestFlags, ESyncStatusRequestFlags::Taskbar))
+			{
+				FillOutTaskbarAutoHide(SyncStatus.Get());
+			}
 
-			if (CpuMonitor)
+			if (EnumHasAnyFlags(RequestFlags, ESyncStatusRequestFlags::PidInFocus))
 			{
-				CpuMonitor->GetPerCoreUtilization(SyncStatus->CpuUtilization);
+				SyncStatus->PidInFocus = FindPidInFocus();
 			}
 
-			const FPlatformMemoryStats MemStats = FPlatformMemory::GetStats();
-			SyncStatus->AvailablePhysicalMemory = MemStats.AvailablePhysical;
+			if (CpuMonitor && EnumHasAnyFlags(RequestFlags, ESyncStatusRequestFlags::CpuUtilization))
+			{
+				CpuMonitor->GetPerCoreUtilization(SyncStatus->CpuUtilization);
+			}
 
-			if (IsNvAPIInitialized)
+			if (EnumHasAnyFlags(RequestFlags, ESyncStatusRequestFlags::AvailablePhysicalMemory))
 			{
-				FillOutPhysicalGpuStats(SyncStatus.Get());
+				const FPlatformMemoryStats MemStats = FPlatformMemory::GetStats();
+				SyncStatus->AvailablePhysicalMemory = MemStats.AvailablePhysical;
 			}
+			
+			// Query GPU stats. These use NVML
+			{
+				const bool bGetUtilizations = EnumHasAnyFlags(RequestFlags, ESyncStatusRequestFlags::GpuUtilization);
+				const bool bGetClocks = EnumHasAnyFlags(RequestFlags, ESyncStatusRequestFlags::GpuCoreClockKhz);
+				const bool bGetTemperatures = EnumHasAnyFlags(RequestFlags, ESyncStatusRequestFlags::GpuTemperature);
 
-			return CreateSyncStatusMessage(SyncStatus.Get());
+				FillOutPhysicalGpuStats(SyncStatus.Get(), bGetUtilizations, bGetClocks, bGetTemperatures);
+			}
+
+			return CreateSyncStatusMessage(SyncStatus.Get(), RequestFlags);
 		}
 	);
 
@@ -2511,12 +2577,16 @@ bool FSwitchboardListener::Task_RefreshMosaics(const FSwitchboardRefreshMosaicsT
 	SWITCHBOARD_TRACE_CPUPROFILER_EVENT_SCOPE(FSwitchboardListener::Task_RefreshMosaics);
 
 #if PLATFORM_WINDOWS
-	if (!bIsNvAPIInitialized)
 	{
-		SendMessage(
-			CreateTaskDeclinedMessage(InRefreshMosaicsTask, "NvAPI not supported", {}),
-			InRefreshMosaicsTask.Recipient);
-		return false;
+		FScopedNvApi ScopedNvApi;
+
+		if (!ScopedNvApi.IsNvApiInitialized())
+		{
+			SendMessage(
+				CreateTaskDeclinedMessage(InRefreshMosaicsTask, "NvAPI not supported", {}),
+				InRefreshMosaicsTask.Recipient);
+			return false;
+		}
 	}
 
 	// Reject request if an equivalent one is already in our future
diff --git a/Engine/Source/Programs/SwitchboardListener/SblCore/Private/SwitchboardProtocol.cpp b/Engine/Source/Programs/SwitchboardListener/SblCore/Private/SwitchboardProtocol.cpp
--- a/Engine/Source/Programs/SwitchboardListener/SblCore/Private/SwitchboardProtocol.cpp
+++ b/Engine/Source/Programs/SwitchboardListener/SblCore/Private/SwitchboardProtocol.cpp
@@ -70,7 +70,7 @@ FString CreateCommandDeclinedMessage(const FGuid& InMessageID, const FString& In
 	return CreateMessage(TEXT("command accepted"), false, { { TEXT("id"), InMessageID.ToString() }, {TEXT("error"), InErrorMessage} });
 }
 
-FString CreateSyncStatusMessage(const FSyncStatus& SyncStatus)
+FString CreateSyncStatusMessage(const FSyncStatus& SyncStatus, ESyncStatusRequestFlags RequestFlags)
 {
 	FString SyncStatusJsonString;
 	const bool bJsonStringOk = FJsonObjectConverter::UStructToJsonObjectString(SyncStatus, SyncStatusJsonString);
@@ -79,11 +79,15 @@ FString CreateSyncStatusMessage(const FSyncStatus& SyncStatus)
 
 	FString Message;
 
+	using RequestFlagsIntType = std::underlying_type_t<ESyncStatusRequestFlags>;
+	RequestFlagsIntType RequestFlagsValue = static_cast<RequestFlagsIntType>(RequestFlags);
+
 	TSharedRef<TJsonWriter<TCHAR, TCondensedJsonPrintPolicy<TCHAR>>> JsonWriter = TJsonWriterFactory<TCHAR, TCondensedJsonPrintPolicy<TCHAR>>::Create(&Message);
 	JsonWriter->WriteObjectStart();
 	JsonWriter->WriteValue(TEXT("get sync status"), true); // TODO: Phase out this field and replace with two below because parser now needs to check all possibilities.
 	JsonWriter->WriteValue(TEXT("command"), TEXT("get sync status"));
 	JsonWriter->WriteValue(TEXT("bAck"), true);
+	JsonWriter->WriteValue(TEXT("request_flags"), RequestFlagsValue);
 	JsonWriter->WriteRawJSONValue(TEXT("syncStatus"), SyncStatusJsonString);
 	JsonWriter->WriteObjectEnd();
 	JsonWriter->Close();
@@ -341,13 +345,43 @@ FCreateTaskResult CreateTaskFromCommand(const FString& InCommand, const FIPv4End
 	{
 		TSharedPtr<FJsonValue> UUIDField = TryGetCommandRequiredField(JsonData, TEXT("uuid"));
 
+		if (!UUIDField.IsValid())
+		{
+			Result.Status = ECreateTaskStatus::Error_ParsingFailed;
+			return Result;
+		}
+
 		FGuid ProgramID;
-		if (UUIDField.IsValid() && FGuid::Parse(UUIDField->AsString(), ProgramID))
+
+		if (!FGuid::Parse(UUIDField->AsString(), ProgramID))
 		{
-			Result.Status = ECreateTaskStatus::Success;
-			Result.Task = MakeUnique<FSwitchboardGetSyncStatusTask>(MessageID, InEndpoint, ProgramID);
+			Result.Status = ECreateTaskStatus::Error_ParsingFailed;
 			return Result;
 		}
+
+		TSharedPtr<FJsonValue> RequestFlagsField = TryGetCommandRequiredField(JsonData, TEXT("request_flags"));
+
+		if (!RequestFlagsField.IsValid())
+		{
+			Result.Status = ECreateTaskStatus::Error_ParsingFailed;
+			return Result;
+		}
+
+		using RequestFlagsIntType = std::underlying_type<ESyncStatusRequestFlags>::type;
+		RequestFlagsIntType RequestFlags;
+
+		if (!RequestFlagsField->TryGetNumber(RequestFlags))
+		{
+			Result.Status = ECreateTaskStatus::Error_ParsingFailed;
+			return Result;
+		}
+
+		const ESyncStatusRequestFlags RequestFlagsEnum = static_cast<ESyncStatusRequestFlags>(RequestFlags);
+
+		Result.Status = ECreateTaskStatus::Success;
+		Result.Task = MakeUnique<FSwitchboardGetSyncStatusTask>(MessageID, InEndpoint, ProgramID, RequestFlagsEnum);
+
+		return Result;
 	}
 	else if (CommandName == FSwitchboardRedeployListenerTask::CommandName)
 	{		
diff --git a/Engine/Source/Programs/SwitchboardListener/SblCore/Private/SwitchboardProtocol.h b/Engine/Source/Programs/SwitchboardListener/SblCore/Private/SwitchboardProtocol.h
--- a/Engine/Source/Programs/SwitchboardListener/SblCore/Private/SwitchboardProtocol.h
+++ b/Engine/Source/Programs/SwitchboardListener/SblCore/Private/SwitchboardProtocol.h
@@ -11,7 +11,7 @@
 struct FSwitchboardPacket;
 struct FSwitchboardTask;
 struct FSyncStatus;
-
+enum class ESyncStatusRequestFlags;
 
 DECLARE_LOG_CATEGORY_EXTERN(LogSwitchboardProtocol, Verbose, All);
 
@@ -38,7 +38,7 @@ FString CreateReceiveFileFromClientFailedMessage(const FString& InDestinationPat
 FString CreateSendFileToClientCompletedMessage(const FString& InSourcePath, const FString& InFileContent);
 FString CreateSendFileToClientFailedMessage(const FString& InSourcePath, const FString& InError);
 
-FString CreateSyncStatusMessage(const FSyncStatus& SyncStatus);
+FString CreateSyncStatusMessage(const FSyncStatus& SyncStatus, ESyncStatusRequestFlags RequestFlags);
 
 FString CreateRedeployStatusMessage(const FGuid& InMessageID, bool bAck, const FString& Status);
 
diff --git a/Engine/Source/Programs/SwitchboardListener/SblCore/Private/SwitchboardTasks.h b/Engine/Source/Programs/SwitchboardListener/SblCore/Private/SwitchboardTasks.h
--- a/Engine/Source/Programs/SwitchboardListener/SblCore/Private/SwitchboardTasks.h
+++ b/Engine/Source/Programs/SwitchboardListener/SblCore/Private/SwitchboardTasks.h
@@ -79,23 +79,51 @@ struct FSwitchboardAuthenticateTask : public FSwitchboardTask
 	//~ End FSwitchboardTask interface
 };
 
+/** Bit flags for the request. It should match its counterpart's definition in message_protocol.py */
+enum class ESyncStatusRequestFlags
+{
+	None = 0,
+	SyncTopos = 1 << 0,
+	MosaicTopos = 1 << 1,
+	FlipModeHistory = 1 << 2,
+	ProgramLayers = 1 << 3,
+	DriverInfo = 1 << 4,
+	Taskbar = 1 << 5,
+	PidInFocus = 1 << 6,
+	CpuUtilization = 1 << 7,
+	AvailablePhysicalMemory = 1 << 8,
+	GpuUtilization = 1 << 9,
+	GpuCoreClockKhz = 1 << 10,
+	GpuTemperature = 1 << 11,
+};
+
 struct FSwitchboardGetSyncStatusTask : public FSwitchboardTask
 {
-	FSwitchboardGetSyncStatusTask(const FGuid& InTaskId, const FIPv4Endpoint& InEndpoint, const FGuid& InProgramID)
+	FSwitchboardGetSyncStatusTask(
+		const FGuid& InTaskId, 
+		const FIPv4Endpoint& InEndpoint, 
+		const FGuid& InProgramID, 
+		const ESyncStatusRequestFlags InRequestFlags
+	)
 		: FSwitchboardTask{ ESwitchboardTaskType::GetSyncStatus, InTaskId, InEndpoint }
 		, ProgramID(InProgramID)
+		, RequestFlags(InRequestFlags)
 	{}
 
 	/** ID of the program that we wish to get the FlipMode of */
 	FGuid ProgramID;
 
+	/** Mask for the information requested */
+	ESyncStatusRequestFlags RequestFlags = static_cast<ESyncStatusRequestFlags>(~0);
+
 	//~ Begin FSwitchboardTask interface
 	static constexpr const TCHAR* CommandName = TEXT("get sync status");
 	virtual const TCHAR* GetCommandName() const override { return CommandName; };
 
 	virtual uint32 GetEquivalenceHash() const override
 	{
-		return HashCombine(FSwitchboardTask::GetEquivalenceHash(), GetTypeHash(ProgramID));
+		uint32 Hash = HashCombine(FSwitchboardTask::GetEquivalenceHash(), GetTypeHash(ProgramID));
+		return HashCombine(Hash, GetTypeHash(RequestFlags));
 	}
 	//~ End FSwitchboardTask interface
 };
diff --git a/Engine/Source/Programs/SwitchboardListener/SblCore/SblCore.Build.cs b/Engine/Source/Programs/SwitchboardListener/SblCore/SblCore.Build.cs
--- a/Engine/Source/Programs/SwitchboardListener/SblCore/SblCore.Build.cs
+++ b/Engine/Source/Programs/SwitchboardListener/SblCore/SblCore.Build.cs
@@ -1,5 +1,6 @@
 // Copyright Epic Games, Inc. All Rights Reserved.
 
+using System.IO;
 using UnrealBuildTool;
 
 public class SblCore : ModuleRules
@@ -60,5 +61,18 @@ public class SblCore : ModuleRules
 		);
 
 		AddEngineThirdPartyPrivateStaticDependencies(Target, "OpenSSL");
+
+		if (Target.Platform == UnrealTargetPlatform.Win64)
+		{
+			PublicAdditionalLibraries.Add(Path.Combine(ModuleDirectory, "NvmlWrapper", "Lib", "Windows", "NvmlWrapper.lib"));
+
+			PublicIncludePaths.AddRange(
+				new string[]
+				{
+				Path.Combine(ModuleDirectory, "NvmlWrapper", "Public"),
+				}
+			);
+		}
+
 	}
 }
