Switchboard: Improved support for content plugins, including uproject AdditionalPluginDirectories.

Switchboard is now aware of all plugins contained in the engine, project, and AdditionalPluginDirectories ("external" plugins), and supports discovering levels and other assets in their Content directories.

By default, all enabled project and external plugins are considered when enumerating content. Other plugins can be included by adding them using the "Content Plugin Filters" setting.

Other changes:
- Fixes an issue where files deleted at head revision would emit spurious "Invalid P4PrintResult" and "No contents for //..." warnings while querying UGS config.
- Fixes usage of deprecated QApplication.exec_().
- Minor linter fixes (import order, max column width).

diff --git a/Engine/Plugins/VirtualProduction/Switchboard/Source/Switchboard/switchboard/__main__.py b/Engine/Plugins/VirtualProduction/Switchboard/Source/Switchboard/switchboard/__main__.py
--- a/Engine/Plugins/VirtualProduction/Switchboard/Source/Switchboard/switchboard/__main__.py
+++ b/Engine/Plugins/VirtualProduction/Switchboard/Source/Switchboard/switchboard/__main__.py
@@ -1,18 +1,21 @@
 # Copyright Epic Games, Inc. All Rights Reserved.
-from switchboard.config import SETTINGS, CONFIG
-from switchboard.switchboard_scripting import ScriptManager
-from .switchboard_dialog import SwitchboardDialog
-from .switchboard_logging import LOGGER
 
+import argparse
 import signal
 import sys
-import argparse
 
 from PySide6 import QtCore, QtWidgets
 
+from switchboard.config import SETTINGS, CONFIG
+from switchboard.switchboard_scripting import ScriptManager
+from .switchboard_dialog import SwitchboardDialog
+from .switchboard_logging import LOGGER
+
+
 # Build resources
 # "D:\P4\UE5\Main\Engine\Extras\ThirdPartyNotUE\SwitchboardThirdParty\Python\Scripts\pyside6-rcc" -o "D:\P4\UE5\Main\Engine\Plugins\VirtualProduction\Switchboard\Source\Switchboard\switchboard\resources.py" "D:\P4\UE5\Main\Engine\Plugins\VirtualProduction\Switchboard\Source\Switchboard\switchboard\ui\resources.qrc"
 
+
 def parse_arguments():
     ''' Parses command line arguments and returns the populated namespace 
     '''
@@ -22,6 +25,7 @@ def parse_arguments():
     args, unknown = parser.parse_known_args()
     return args
 
+
 def launch():
     """
     Main for running standalone or in another application.
@@ -85,7 +89,7 @@ def launch():
     timer.timeout.connect(lambda: None)
 
     # execute the app
-    appresult = app.exec_()
+    appresult = app.exec()
 
     # script exit
     script_manager.on_exit()
diff --git a/Engine/Plugins/VirtualProduction/Switchboard/Source/Switchboard/switchboard/config.py b/Engine/Plugins/VirtualProduction/Switchboard/Source/Switchboard/switchboard/config.py
--- a/Engine/Plugins/VirtualProduction/Switchboard/Source/Switchboard/switchboard/config.py
+++ b/Engine/Plugins/VirtualProduction/Switchboard/Source/Switchboard/switchboard/config.py
@@ -1,19 +1,19 @@
 # Copyright Epic Games, Inc. All Rights Reserved.
 
 import collections
+from enum import Enum
 import fnmatch
 import io
 import json
 import os
 import pathlib
 import shutil
+from pathlib import Path, PurePosixPath
 import socket
 import sys
 import threading
 import time
 from typing import Any, Callable, Optional, Tuple, Type, Union
-from enum import Enum
-from pathlib import Path
 
 from PySide6 import QtCore
 from PySide6 import QtGui
@@ -23,8 +23,9 @@ from switchboard import switchboard_widgets as sb_widgets
 from switchboard.switchboard_logging import LOGGER
 from switchboard.switchboard_widgets import (
     DropDownMenuComboBox, NonScrollableComboBox)
-from switchboard import ue_plugin_utils, ugs_utils
+from switchboard import ugs_utils
 from switchboard.sbcache import SBCache, Map
+from switchboard.ue_plugin_utils import UnrealPlugin, UnrealPluginManager
 
 ROOT_CONFIGS_PATH = pathlib.Path(__file__).parent.with_name('configs')
 CONFIG_SUFFIX = '.json'
@@ -357,7 +358,7 @@ class Setting(QtCore.QObject):
                 setting_label.setToolTip(self.tool_tip)
 
             form_layout.addRow(
-                setting_label,                
+                setting_label,
                 self._decorate_with_reset_widget(override_device_name, top_level_widget)
             )
 
@@ -376,7 +377,12 @@ class Setting(QtCore.QObject):
         )
 
     def _on_widget_destroyed(self, on_setting_changed_lambda, override_device_name: str):
-        self.signal_setting_changed.disconnect(on_setting_changed_lambda)
+        try:
+            self.signal_setting_changed.disconnect(on_setting_changed_lambda)
+        except Exception as exc:
+            LOGGER.warning('Failed to disconnect on_setting_changed_lambda',
+                           exc_info=exc)
+
         self.set_widget(widget=None, override_device_name=override_device_name)
 
     def _decorate_with_reset_widget(self, override_device_name: str, setting_editor_widget: QtWidgets.QWidget):
@@ -2076,7 +2082,7 @@ class Config(object):
         ''' Restores saving_allowed flag from the stack
         '''
         self.saving_allowed = self.saving_allowed_fifo.pop()
-        
+
     def init(self, file_path: Union[str, pathlib.Path]):
         self.init_with_file_path(file_path)
 
@@ -2088,8 +2094,8 @@ class Config(object):
                 # Read the json config file
                 with open(self.file_path, 'r') as f:
                     LOGGER.debug(f'Loading Config {self.file_path}')
-                    data = json.load(f)                    
-                        
+                    data = json.load(f)
+
             except (ConfigPathError, FileNotFoundError) as e:
                 LOGGER.error(f'Config: {e}')
                 self.file_path = None
@@ -2129,6 +2135,18 @@ class Config(object):
         # MISC SETTINGS
         self.CURRENT_LEVEL = data.get('current_level', DEFAULT_MAP_TEXT)
 
+        # UE plugin tracking
+        self.ue_plugin_mgr = UnrealPluginManager()
+        self.ue_plugin_mgr.set_engine_dir(Path(self.ENGINE_DIR.get_value()))
+        self.ue_plugin_mgr.set_uproject_path(
+            Path(self.UPROJECT_PATH.get_value()))
+
+        self.ENGINE_DIR.signal_setting_changed.connect(
+            lambda _, new: self.ue_plugin_mgr.set_engine_dir(Path(new)))
+
+        self.UPROJECT_PATH.signal_setting_changed.connect(
+            lambda _, new: self.ue_plugin_mgr.set_uproject_path(Path(new)))
+
         # Devices
         self._device_data_from_config = {}
         self._plugin_data_from_config = {}
@@ -2206,7 +2224,7 @@ class Config(object):
 
         SETTINGS.CONFIG = self.file_path
         SETTINGS.save()
-        
+
     def init_switchboard_settings(self, data={}):
         self.switchboard_settings = {
             "listener_exe": StringSetting(
@@ -2215,7 +2233,7 @@ class Config(object):
                 value = data.get('listener_exe', 'SwitchboardListener')
             )
         }
-        
+
         self.LISTENER_EXE = self.switchboard_settings["listener_exe"]
 
     def init_sblhelper_settings(self, data={}):
@@ -2290,11 +2308,12 @@ class Config(object):
                 tool_tip=(
                     "Plugins that match any of these filters will also be "
                     "searched when populating fields in Switchboard (e.g. "
-                    "levels, nDisplay configs, etc.). Each value can be "
-                    "either just a plugin name identifying a project plugin, "
-                    "or a relative or absolute path to a plugin directory. "
-                    "Relative paths should be relative to the directory "
-                    "containing the .uproject file."),
+                    "levels, nDisplay configs, etc.).\n"
+                    "\n"
+                    "Each value can be either a plugin name, or a relative or "
+                    "absolute path to a plugin directory. Relative paths "
+                    "should be relative to the directory containing the "
+                    ".uproject file."),
                 migrate_data=migrate_comma_separated_string_to_list
             ),
         }
@@ -2706,11 +2725,6 @@ class Config(object):
         del self._device_settings[(device_type, device_name)]
         self.save()
 
-    def shrink_path(self, path):
-        path_name = path.replace(self.get_project_dir(), '', 1)
-        path_name = path_name.replace(os.sep, '/')
-        return path_name
-
     def get_project_dir(self) -> str:
         '''
         Get the root directory of the project.
@@ -2725,47 +2739,48 @@ class Config(object):
         '''
         return os.path.join(self.get_project_dir(), 'Content')
 
-    def get_unreal_content_plugins(
-            self) -> list[ue_plugin_utils.UnrealPlugin]:
+    def get_unreal_content_plugins(self) -> list[UnrealPlugin]:
         '''
         Get a list of Unreal Engine plugins that match the current
         Switchboard config's content plugin filter settings.
 
-        By default, Switchboard does not search inside plugins for assets when
-        populating UI fields such as the level or nDisplay config dropdown
-        menus. Plugins in which to search for content can be selectively
-        added though using the "Content Plugin Filters" setting, which
-        stores a list of glob-style patterns.
+        By default, Switchboard searches inside project plugins as well as
+        AdditionalPluginDirectories plugins for assets when populating UI
+        fields such as the level or nDisplay config dropdown menus.
+
+        Engine/other plugins in which to search for content can be selectively
+        added though using the "Content Plugin Filters" setting.
         '''
-        filter_patterns = self.CONTENT_PLUGIN_FILTERS.get_value()
-        unreal_plugins = ue_plugin_utils.UnrealPlugin.from_path_filters(
-            self.get_project_dir(), filter_patterns)
+
+        # Pop known plugin names from the config array, treat the rest as globs
+        unreal_plugins: list[UnrealPlugin] = []
+        filter_patterns: list[str] = []
+
+        # Enabled project plugins are always enumerated for content
+        unreal_plugins.extend(self.ue_plugin_mgr.enabled_project_plugins)
+
+        for pattern in self.CONTENT_PLUGIN_FILTERS.get_value():
+            if plugin := self.ue_plugin_mgr.get_plugin_by_name(pattern):
+                unreal_plugins.append(plugin)
+            else:
+                filter_patterns.append(pattern)
+
+        if filter_patterns:
+            unreal_plugins.extend(UnrealPlugin.from_path_filters(
+                self.UPROJECT_PATH.get_value(), filter_patterns))
+
         return unreal_plugins
 
     def resolve_content_path(
-            self,
-            file_path: str,
-            unreal_content_plugin: ue_plugin_utils.UnrealPlugin = None) -> str:
+        self,
+        file_path: Union[Path, str]
+    ) -> Optional[PurePosixPath]:
         '''
         Resolve a file path on the file system to the corresponding content
         path in UE.
-
-        If an unreal_content_plugin is provided, the file is assumed to live
-        inside that plugin and its content path will have the appropriate
-        plugin name-based prefix. Otherwise, the file is assumed to live
-        inside the project's content folder.
         '''
-        if unreal_content_plugin:
-            content_dir = str(unreal_content_plugin.plugin_content_path)
-            ue_path_prefix = str(unreal_content_plugin.mounted_path)
-        else:
-            content_dir = self.get_project_content_dir()
-            ue_path_prefix = '/Game'
-
-        path_name = file_path.replace(content_dir, ue_path_prefix, 1)
-        path_name = self.shrink_path(path_name)
 
-        return path_name
+        return self.ue_plugin_mgr.file_to_content_path(file_path)
 
     def find_levels(self) -> list[str]:
 
@@ -2824,7 +2839,9 @@ class Config(object):
         # (unreal_plugin, directory_path). This allows us to differentiate
         # between maps in the project (unreal_plugin is None in that case) and
         # maps in a plugin.
-        search_paths = [(None, project_maps_path)]
+        search_paths: list[tuple[Optional[UnrealPlugin], Path]] = [
+            (None, Path(project_maps_path))
+        ]
 
         for unreal_content_plugin in self.get_unreal_content_plugins():
             search_paths.append(
@@ -2868,14 +2885,12 @@ class Config(object):
                 if not fnmatch.fnmatch(umap.name, maps_filter):
                     continue
 
-                mapgamepath = Path(self.resolve_content_path(
-                    umap.path,
-                    unreal_content_plugin=unreal_content_plugin))
-
-                mapgamepath = (mapgamepath.parent / mapgamepath.stem).as_posix()
-
-                if mapgamepath not in levels:
-                    levels.append(mapgamepath)
+                if mapgamepath := self.resolve_content_path(umap.path):
+                    mapgamepath = mapgamepath.with_suffix('')
+                    if mapgamepath not in levels:
+                        levels.append(str(mapgamepath))
+                else:
+                    LOGGER.warning(f"No game path for {umap.path}")
 
         levels.sort()
 
diff --git a/Engine/Plugins/VirtualProduction/Switchboard/Source/Switchboard/switchboard/devices/ndisplay/plugin_ndisplay.py b/Engine/Plugins/VirtualProduction/Switchboard/Source/Switchboard/switchboard/devices/ndisplay/plugin_ndisplay.py
--- a/Engine/Plugins/VirtualProduction/Switchboard/Source/Switchboard/switchboard/devices/ndisplay/plugin_ndisplay.py
+++ b/Engine/Plugins/VirtualProduction/Switchboard/Source/Switchboard/switchboard/devices/ndisplay/plugin_ndisplay.py
@@ -3,7 +3,7 @@
 import json
 import os
 import sys
-from pathlib import Path
+from pathlib import Path, PurePosixPath
 import socket
 import struct
 import traceback
@@ -15,7 +15,6 @@ from PySide6 import QtWidgets
 from switchboard import message_protocol, switchboard_application
 from switchboard import switchboard_utils as sb_utils
 from switchboard import switchboard_widgets as sb_widgets
-from switchboard import switchboard_dialog as sb_dialog
 from switchboard.config import CONFIG, BoolSetting, IntSetting, FilePathSetting, \
     LoggingSetting, OptionSetting, Setting, StringSetting, SETTINGS, \
     StringListSetting, AddressSetting, migrate_comma_separated_string_to_list, \
@@ -176,10 +175,6 @@ class AddnDisplayDialog(AddDeviceDialog):
         if len(cfg_path) > 0 and os.path.exists(cfg_path):
             self.cbConfigs.setCurrentText(cfg_path)
 
-    def generate_short_unique_config_name(config_path: str, file_name: str) -> str:
-        config_path = CONFIG.shrink_path(config_path)
-        return sb_dialog.SwitchboardDialog.filter_empty_abiguated_path(config_path, file_name)
-
     def on_clicked_btnFindConfigs(self):
         ''' Finds and populates config combobox '''
 
@@ -1416,26 +1411,6 @@ class DevicenDisplay(DeviceUnreal):
         self.settings['config_graphics_adapter'].update_value(
             menode['kwargs'].get("config_graphics_adapter", -1))
 
-    @classmethod
-    def uasset_path_from_object_path(
-            cls, object_path: str, project_dir: str) -> str:
-        '''
-        Given a full object path, return the package file path, treating
-        "`project_dir`/Content/" as "/Game/".
-        '''
-        expected_root = '/Game/'
-        if not object_path.startswith(expected_root):
-            raise ValueError('Unsupported object path root')
-
-        # If object_path is: /Game/PathA/PathB/Package.Object:SubObject
-        # Then package_rel_path is: PathA/PathB/Package
-        path_end_idx = object_path.rindex('/')
-        package_end_idx = object_path.index('.', path_end_idx)
-        package_rel_path = object_path[len(expected_root):package_end_idx]
-
-        return os.path.normpath(
-            os.path.join(project_dir, 'Content', f'{package_rel_path}.uasset'))
-
     def get_connected_devices(self):
         ''' Returns a list with the connected devices/nodes
         '''
@@ -1525,17 +1500,25 @@ class DevicenDisplay(DeviceUnreal):
         self.unreal_client.send_message(cfg_msg)
 
         if self.bp_object_path:
-            local_uasset_path = self.uasset_path_from_object_path(
-                self.bp_object_path, os.path.dirname(
-                    CONFIG.UPROJECT_PATH.get_value()))
-            dest_uasset_path = self.uasset_path_from_object_path(
-                self.bp_object_path, os.path.dirname(
-                    CONFIG.UPROJECT_PATH.get_value(self.name)))
+            local_project_path = os.path.dirname(
+                CONFIG.UPROJECT_PATH.get_value())
+            dest_project_path = os.path.dirname(
+                CONFIG.UPROJECT_PATH.get_value(self.name))
+
+            local_uasset_path = CONFIG.ue_plugin_mgr.content_to_file_path(
+                PurePosixPath(self.bp_object_path)).with_suffix('.uasset')
+
+            uasset_project_relative = os.path.relpath(local_uasset_path,
+                                                      local_project_path)
+
+            dest_uasset_path = os.path.join(dest_project_path,
+                                            uasset_project_relative)
 
             if os.path.isfile(local_uasset_path):
                 self.pending_transfer_uasset = True
                 _, uasset_msg = message_protocol.create_send_file_message(
-                    local_uasset_path, dest_uasset_path, force_overwrite=True)
+                    str(local_uasset_path), dest_uasset_path,
+                    force_overwrite=True)
                 self.unreal_client.send_message(uasset_msg)
             else:
                 LOGGER.warning(
diff --git a/Engine/Plugins/VirtualProduction/Switchboard/Source/Switchboard/switchboard/devices/unreal/plugin_unreal.py b/Engine/Plugins/VirtualProduction/Switchboard/Source/Switchboard/switchboard/devices/unreal/plugin_unreal.py
--- a/Engine/Plugins/VirtualProduction/Switchboard/Source/Switchboard/switchboard/devices/unreal/plugin_unreal.py
+++ b/Engine/Plugins/VirtualProduction/Switchboard/Source/Switchboard/switchboard/devices/unreal/plugin_unreal.py
@@ -9,9 +9,11 @@ from datetime import datetime
 from functools import wraps
 import hashlib
 from ipaddress import IPv4Address
+from itertools import count
 import json
 import os
 import pathlib
+from pathlib import PurePosixPath
 import re
 import socket
 import sys
@@ -373,9 +375,13 @@ class LiveLinkPresetSetting(Setting):
             project=project,
             classnames=self._classnames())
 
+        assets = sorted(assets, key=lambda x: x.name)
+
         # generate the combo box items
-        for asset in assets:
-            combo.addItem(asset.name.replace('.uasset', ''), asset)
+        for idx, asset in zip(count(1), assets):
+            combo.addItem(asset.name.removesuffix('.uasset'), asset)
+            combo.setItemData(idx, asset.gamepath.removesuffix('.uasset'),
+                              QtCore.Qt.ItemDataRole.ToolTipRole)
 
         # set the current index to the live link preset that was already selected
         for item_idx in range(combo.count()):
@@ -1208,7 +1214,7 @@ class DeviceUnreal(Device):
 
         return valid
 
-    def conform_asset_gamepath(self, gamepath : str, ext : str = '.uasset') -> str:
+    def conform_asset_gamepath(self, gamepath: str, ext: str = '.uasset') -> str:
         ''' Conforms the given gamepath, which can include file extension, to the //Game/../MyAsset.MyAsset convention
         e.g.
             /Game/Folder/MyAssetName.uasset -> /Game/Folder/MyAssetName.MyAssetName
@@ -1221,7 +1227,7 @@ class DeviceUnreal(Device):
 
         return f"{gamepath}.{name}"
 
-    def exec_command_for_livelink_preset(self, livelink_preset_gamepath : str) -> str:
+    def exec_command_for_livelink_preset(self, livelink_preset_gamepath: str) -> str:
         ''' Returns the exec command string to enable the given livelink preset gamepath
         LiveLink presets can be applied as ExecCmds
         e.g. of command:
@@ -1233,7 +1239,7 @@ class DeviceUnreal(Device):
 
         return f"LiveLink.Preset.Apply Preset={self.conform_asset_gamepath(livelink_preset_gamepath)}"
 
-    def dpcvar_for_mediaprofile(self, mediaprofile_gamepath : str) -> str:
+    def dpcvar_for_mediaprofile(self, mediaprofile_gamepath: str) -> str:
         ''' Returns the dpcvar assignment string to enable the given mediaprofile gamepath
         Media profiles can be applied as early cvars
         e.g.:
@@ -2896,24 +2902,14 @@ class DeviceUnreal(Device):
         e.g. nDisplay configs, live link presets, and media profiles.
         '''
 
-        project_configs_path = os.path.normpath(
-            CONFIG.get_project_content_dir())
-
-        # search_paths stores a list of tuples of the form
-        # (unreal_plugin, directory_path). This allows us to differentiate
-        # between project assets (unreal_plugin is None in
-        # that case) and plugin assets.
-        search_paths = [(None, project_configs_path)]
+        search_paths: list[pathlib.Path] = []
+        search_paths.append(pathlib.Path(CONFIG.get_project_content_dir()))
+        search_paths.extend([plugin.plugin_content_path
+                            for plugin in CONFIG.get_unreal_content_plugins()])
 
-        for unreal_content_plugin in CONFIG.get_unreal_content_plugins():
-            search_paths.append(
-                (unreal_content_plugin,
-                 unreal_content_plugin.plugin_content_path))
-
-        asset_names = []
-        asset_paths = []
-        asset_plugins = []
-        asset_classnames = []
+        asset_names: list[str] = []
+        asset_paths: list[str] = []
+        asset_classnames: list[str] = []
 
         assets = []
 
@@ -2935,8 +2931,8 @@ class DeviceUnreal(Device):
         # convenience find assets worker. This is a step before parsing the assets, and can
         # also take some time depending on the number of assets in the project.
         def find_assets_work():
-            for (unreal_content_plugin, configs_path) in search_paths:
-                for dirpath, _, file_names in os.walk(configs_path):
+            for path in search_paths:
+                for dirpath, _, file_names in os.walk(path):
                     for file_name in file_names:
                         if not file_name.lower().endswith(('.uasset', '.ndisplay')):
                             continue
@@ -2951,12 +2947,10 @@ class DeviceUnreal(Device):
                                 assets.append({
                                     'name': file_name,
                                     'path': asset_path,
-                                    'plugin': unreal_content_plugin,
                                 })
                             else:
                                 asset_names.append(file_name)
                                 asset_paths.append(asset_path)
-                                asset_plugins.append(unreal_content_plugin)
                                 asset_classnames.append(DeviceUnreal.NDISPLAY_CLASS_NAMES[0]) # so that it passes the filter later on
 
             done_event.set()
@@ -3031,13 +3025,11 @@ class DeviceUnreal(Device):
                     # first make sure they all exist
                     name = asset['name']
                     path = asset['path']
-                    plugin = asset['plugin']
                     classname = asset['assetdata'].ObjectClassName
 
                     # now append to lists
                     asset_names.append(name)
                     asset_paths.append(path)
-                    asset_plugins.append(plugin)
                     asset_classnames.append(classname)
 
                 except Exception:
@@ -3046,12 +3038,26 @@ class DeviceUnreal(Device):
         # close progress bar window
         progressDiag.close()
 
-        def generate_short_unique_config_name(config_path: str, file_name: str) -> str:
-            config_path = CONFIG.shrink_path(config_path)
-            return sb_dialog.SwitchboardDialog.filter_empty_abiguated_path(config_path, file_name)
+        asset_path_to_game_path: dict[str, Optional[PurePosixPath]] = {
+            x: CONFIG.resolve_content_path(x) for x in asset_paths
+        }
+
+        def disambiguate_asset_path(path: str) -> str:
+            if game_path := asset_path_to_game_path[path]:
+                file_name = game_path.stem
+                rel_path = game_path.parent
+            else:
+                LOGGER.error(f"Couldn't map game path for asset {path}")
+                file_name = os.path.basename(path)
+                rel_path = path.removesuffix(file_name)
+
+            if rel_path == "/":
+                return f"{file_name}"
+            else:
+                return f"{file_name} ({rel_path})"
 
-        asset_names, _ = sb_dialog.SwitchboardDialog.generate_disambiguated_names(
-            asset_paths, generate_short_unique_config_name)
+        asset_names, _ = sb_dialog.SwitchboardDialog.disambiguated_base_names(
+            asset_paths, disambiguate_asset_path)
 
         # collect the found config files into the assets list
 
@@ -3060,17 +3066,16 @@ class DeviceUnreal(Device):
         project = SBCache().query_or_create_project(CONFIG.UPROJECT_PATH.get_value())
 
         for idx, asset_name in enumerate(asset_names):
-
-            gamepath = CONFIG.resolve_content_path(
-                file_path=asset_paths[idx], 
-                unreal_content_plugin=asset_plugins[idx])
+            gamepath = asset_path_to_game_path.get(asset_paths[idx])
+            if gamepath is None:
+                continue  # We will have logged an error already.
 
             assettype = SBCache().query_or_create_assettype(asset_classnames[idx])
             asset = Asset(
                 id=0,
                 project=project,
                 assettype=assettype,
-                gamepath=gamepath,
+                gamepath=str(gamepath),
                 name=asset_name,
                 localpath=asset_paths[idx])
 
diff --git a/Engine/Plugins/VirtualProduction/Switchboard/Source/Switchboard/switchboard/switchboard_dialog.py b/Engine/Plugins/VirtualProduction/Switchboard/Source/Switchboard/switchboard/switchboard_dialog.py
--- a/Engine/Plugins/VirtualProduction/Switchboard/Source/Switchboard/switchboard/switchboard_dialog.py
+++ b/Engine/Plugins/VirtualProduction/Switchboard/Source/Switchboard/switchboard/switchboard_dialog.py
@@ -11,7 +11,7 @@ import time
 import re
 import sys
 import json
-from typing import List, Optional, Set, Union
+from typing import Callable, List, Optional, Set, Union
 
 from pathlib import Path
 
@@ -25,12 +25,13 @@ from PySide6.QtWidgets import QWidgetAction, QMenu
 from switchboard import config
 from switchboard import config_osc as osc
 from switchboard import p4_utils
-from switchboard import ugs_utils
 from switchboard import recording
 from switchboard import resources  # noqa
 from switchboard import switchboard_application
 from switchboard import switchboard_utils
 from switchboard import switchboard_widgets as sb_widgets
+from switchboard import ue_plugin_utils
+from switchboard import ugs_utils
 from switchboard.add_config_dialog import AddConfigDialog
 from switchboard.config import CONFIG, DEFAULT_MAP_TEXT, ENABLE_UGS_SUPPORT, SETTINGS, EngineSyncMethod
 from switchboard.device_list_widget import DeviceListWidget, DeviceWidgetHeader
@@ -2046,29 +2047,29 @@ class SwitchboardDialog(QtCore.QObject):
         if self.level != current_level:
             CONFIG.save()
 
-    def filter_empty_abiguated_path(path, file_name):
-        path = path.removesuffix(file_name)
-        if path == "/":
-            return f"{file_name}"
-        else:
-            return f"{file_name} ({path})"
-
-    def generate_short_map_path(path: str, file_name: str) -> str:
-        path = path.replace("/Game", "", 1)
-        return SwitchboardDialog.filter_empty_abiguated_path(path, file_name)
-
-    def generate_disambiguated_names(path_list, shortening_function):
-        name_counts = {}
+    @classmethod
+    def disambiguated_base_names(
+        cls,
+        path_list: list[str],
+        disambiguating_function: Callable[[str], str],
+    ) -> tuple[list[str], dict[str, str]]:
+        '''
+        By default, shorten each path in `path_list` to its basename.
+        If a given basename appears more than once, `disambiguating_function`
+        is called to generate an alternative label (e.g. "A (/Other/Path/)").
+        '''
+        name_counts: dict[str, int] = {}
         for path in path_list:
             file_name = os.path.basename(path)
             name_counts[file_name] = name_counts.get(file_name, 0) + 1
 
-        # Show only level name if unique and show path behind to disambiguate duplicates
-        short_name_list = []
-        short_name_to_path = {}
+        # If unique, reduce to basename. Otherwise, append path to end
+        short_name_list: list[str] = []
+        short_name_to_path: dict[str, str] = {}
         for path in path_list:
             file_name = os.path.basename(path)
-            short_name = file_name if name_counts[file_name] == 1 else shortening_function(path, file_name)
+            short_name = (file_name if name_counts[file_name] == 1
+                          else disambiguating_function(path))
             short_name_list.append(short_name)
             short_name_to_path[short_name] = path
 
@@ -2079,9 +2080,16 @@ class SwitchboardDialog(QtCore.QObject):
             return -1 if path_a.lower() < path_b.lower() \
                 else 1 if path_a.lower() > path_b.lower() else 0
 
-        short_name_list, short_name_to_path = SwitchboardDialog.generate_disambiguated_names(
-            level_path_list,
-            SwitchboardDialog.generate_short_map_path)
+        def disambiguate_level_path(path: str) -> str:
+            file_name = os.path.basename(path)
+            path = path.removesuffix(file_name)
+            if path == "/":
+                return f"{file_name}"
+            else:
+                return f"{file_name} ({path})"
+
+        short_name_list, short_name_to_path = SwitchboardDialog.disambiguated_base_names(
+            level_path_list, disambiguate_level_path)
 
         from functools import cmp_to_key
         short_name_list = sorted(short_name_list, key=cmp_to_key(compare_file_names))
diff --git a/Engine/Plugins/VirtualProduction/Switchboard/Source/Switchboard/switchboard/ue_plugin_utils.py b/Engine/Plugins/VirtualProduction/Switchboard/Source/Switchboard/switchboard/ue_plugin_utils.py
--- a/Engine/Plugins/VirtualProduction/Switchboard/Source/Switchboard/switchboard/ue_plugin_utils.py
+++ b/Engine/Plugins/VirtualProduction/Switchboard/Source/Switchboard/switchboard/ue_plugin_utils.py
@@ -1,10 +1,280 @@
 # Copyright Epic Games, Inc. All Rights Reserved.
 
+from __future__ import annotations
+
+from enum import auto, Enum
 import glob
-import pathlib
-import typing
+import json
+import logging
+from pathlib import Path, PurePosixPath
+import os
+import time
+from typing import Any, Optional, Union
+
+
+PROJECT_CONTENT_MOUNT_NAME = 'Game'
+PROJECT_CONTENT_MOUNT_POINT = f'/{PROJECT_CONTENT_MOUNT_NAME}'
+
+
+class UnrealPluginType(Enum):
+    ''' Mirrors `EPluginType` from IPluginManager.h. '''
+    Engine = auto()
+    Enterprise = auto()
+    Project = auto()
+    External = auto()
+    Mod = auto()
+
+
+class UnrealPluginManager:
+    '''
+    Enumerates the engine and project plugin directories, as well as any paths
+    listed in the .uproject "AdditionalPluginDirectories" field (so-called
+    `UnrealPluginType.External` plugins) to discover .uplugin descriptors.
+
+    Also has functionality to map local file paths to Unreal content mount
+    points (`file_to_content_path`), and vice versa (`content_to_file_path`).
+
+    Any `Project` or `External` plugins which are not marked disabled are
+    returned by the `enabled_project_plugins` property accessor, and by default
+    content within those plugins is surfaced elsewhere in Switchboard.
+    '''
+
+    def __init__(self):
+        self._engine_plugin_paths: list[Path] = []
+        self._project_plugin_paths: list[Path] = []
+        self._external_plugin_paths: list[Path] = []
+
+        self._engine_dir: Optional[Path] = None
+        self._uproject_path: Optional[Path] = None
+
+        # Information derived from .uproject
+        self._addl_plugin_dirs: list[Path] = []
+        self._disabled_plugins: set[str] = set()
+
+        # Plugin name -> UnrealPlugin
+        self._plugin_map: dict[str, UnrealPlugin] = {}
+
+        # Reverse lookup from Content dir to owning plugin
+        self._content_dir_map: dict[Path, UnrealPlugin] = {}
+
+        # Cached for quick reference; also includes "external" plugins
+        self._enabled_project_plugins: list[UnrealPlugin] = []
+
+    @property
+    def engine_dir(self) -> Optional[Path]:
+        return self._engine_dir
+
+    @property
+    def uproject_path(self) -> Optional[Path]:
+        return self._uproject_path
+
+    def set_engine_dir(self, in_dir: Optional[Path]):
+        if self._engine_dir == in_dir:
+            return
+
+        self._engine_plugin_paths.clear()
+        self._engine_dir = in_dir
+        if in_dir:
+            self._engine_plugin_paths = self.find_plugins_in_dir(
+                in_dir / 'Plugins')
+
+        self._refresh_plugin_map()
+
+    def set_uproject_path(self, in_path: Optional[Path]):
+        if self.uproject_path == in_path:
+            return
+
+        self._project_plugin_paths.clear()
+        self._external_plugin_paths.clear()
+
+        self._addl_plugin_dirs.clear()
+        self._disabled_plugins.clear()
+        self._enabled_project_plugins.clear()
+
+        self._uproject_path = in_path
+        if not self._uproject_path:
+            return
+
+        try:
+            with open(self._uproject_path, encoding='utf-8') as f:
+                uproj: dict[str, Any] = json.load(f)
+        except Exception as exc:
+            logging.error('Error parsing .uproject JSON for %s',
+                          self._uproject_path, exc_info=exc)
+            self._uproject_path = None
+            return
+
+        proj_dir = self._uproject_path.parent
+        proj_plugin_path = proj_dir / 'Plugins'
+        self._project_plugin_paths = self.find_plugins_in_dir(proj_plugin_path)
+
+        # Parse .uproject additional plugin paths
+        addl_dirs: list[str] = uproj.get('AdditionalPluginDirectories', [])
+        self._addl_plugin_dirs = [proj_dir / x for x in addl_dirs]
+
+        for addl_dir in self._addl_plugin_dirs:
+            addl_plugins = self.find_plugins_in_dir(addl_dir)
+            self._external_plugin_paths.extend(addl_plugins)
+
+        # Parse .uproject plugin references to identify disabled plugins
+        proj_plugin_refs: list[dict[str, Any]] = uproj.get('Plugins', [])
+        for ref in proj_plugin_refs:
+            ref_name: Optional[str] = ref.get('Name')
+            if ref_name is None:
+                logging.error('Error parsing %s: '
+                              'Plugin references must have a "Name" field',
+                              in_path)
+                continue
+
+            ref_enabled: Optional[bool] = ref.get('Enabled')
+            if ref_enabled is None:
+                logging.error('Error parsing %s: '
+                              'Plugin references must have an "Enabled" field',
+                              in_path)
+                continue
+
+            if not ref_enabled:
+                self._disabled_plugins.add(ref_name)
+
+        # Update plugin map
+        self._refresh_plugin_map()
+
+        # Cache enabled project + external plugins
+        for plugin_path in (self._project_plugin_paths +
+                            self._external_plugin_paths):
+            plugin_name = plugin_path.stem
+
+            if plugin_name in self._disabled_plugins:
+                continue
+
+            if plugin := self.get_plugin_by_name(plugin_name):
+                self._enabled_project_plugins.append(plugin)
+            else:
+                logging.warning("Couldn't find enabled project plugin %s",
+                                plugin_path)
+
+    def content_to_file_path(
+        self,
+        path: Union[PurePosixPath, str]
+    ) -> Optional[Path]:
+        '''
+        Given an Unreal content path beginning with e.g. `/Game/...` or
+        `/PluginName/...` map it to the corresponding local file or directory.
+        '''
+
+        if isinstance(path, str):
+            path = PurePosixPath(path)
+
+        if not path.is_absolute() or len(path.parts) < 2:
+            logging.error('Not an absolute path: %s', path)
+            return None
+
+        mount_name = path.parts[1]
+        rest = path.relative_to(f'/{mount_name}/')
+
+        if rest.name:
+            # Trim object/subobject portion if present
+            if (dot_idx := rest.name.find('.')) != -1:
+                rest = rest.parent / rest.name[:dot_idx]
+
+        if mount_name == PROJECT_CONTENT_MOUNT_POINT:
+            assert self._uproject_path
+            return self._uproject_path.parent / 'Content' / rest
+        else:
+            plugin = self.get_plugin_by_name(mount_name)
+            if plugin is None:
+                logging.warning('Unknown mount point for path %s', path)
+                return None
+
+            return plugin.plugin_content_path / rest
+
+    def file_to_content_path(
+        self,
+        path: Union[Path, str]
+    ) -> Optional[PurePosixPath]:
+        '''
+        Given a file or directory path, map it to the corresponding Unreal
+        content path (relative to the project or plugin content mount point).
+        '''
+
+        if isinstance(path, str):
+            path = Path(path)
+
+        assert self._uproject_path
+        project_content_path = self._uproject_path.parent / 'Content'
+
+        if path.is_relative_to(project_content_path):
+            rest = path.relative_to(project_content_path)
+            return PurePosixPath(PROJECT_CONTENT_MOUNT_POINT) / rest
+
+        for parent in path.parents:
+            if plugin := self.get_plugin_by_content_dir(parent):
+                rest = path.relative_to(parent)
+                return plugin.mounted_path / rest
+
+        logging.warning('Failed to resolve content path "%s"', path)
+        return None
+
+    def get_plugin_by_name(self, name: str) -> Optional[UnrealPlugin]:
+        return self._plugin_map.get(name)
 
-from switchboard.switchboard_logging import LOGGER
+    def get_plugin_by_content_dir(
+            self,
+            in_dir: Path
+    ) -> Optional[UnrealPlugin]:
+        return self._content_dir_map.get(in_dir)
+
+    def find_plugins_in_dir(self, in_path: Path) -> list[Path]:
+        logging.debug('Discovering plugins in path "%s"', in_path)
+
+        start_time = time.perf_counter()
+
+        found_plugins: list[Path] = []
+
+        for dirpath, dirs, files in os.walk(in_path):
+            for file in files:
+                if file.casefold().endswith('.uplugin'):
+                    found_plugins.append(Path(dirpath) / file)
+                    dirs.clear()  # Plugins don't nest; don't descend this tree
+                    continue
+
+        end_time = time.perf_counter()
+
+        logging.debug('Discovered %i plugins in %.2f seconds',
+                      len(found_plugins), end_time - start_time)
+
+        return found_plugins
+
+    @property
+    def enabled_project_plugins(self) -> list[UnrealPlugin]:
+        return self._enabled_project_plugins
+
+    def _refresh_plugin_map(self):
+        self._plugin_map.clear()
+        self._content_dir_map.clear()
+
+        # Plugin name collisions in later paths supercede earlier ones.
+        plugin_lists: list[tuple[list[Path], UnrealPluginType]] = [
+            (self._engine_plugin_paths, UnrealPluginType.Engine),
+            (self._external_plugin_paths, UnrealPluginType.External),
+            (self._project_plugin_paths, UnrealPluginType.Project),
+        ]
+
+        for plugin_list, plugin_type in plugin_lists:
+            for plugin_path in plugin_list:
+                plugin_name = plugin_path.stem
+
+                # TODO?: Precedence logic in FPluginManager::CreatePluginObject
+                if existing := self._plugin_map.get(plugin_name):
+                    logging.warning('Plugin at "%s" superceded by "%s"',
+                                    existing.uplugin_file_path, plugin_path)
+
+                plugin = UnrealPlugin(plugin_path, plugin_type=plugin_type)
+
+                self._plugin_map[plugin_name] = plugin
+
+        self._content_dir_map = {x.plugin_content_path: x
+                                 for x in self._plugin_map.values()}
 
 
 class UnrealPlugin(object):
@@ -14,7 +284,10 @@ class UnrealPlugin(object):
     '''
 
     @classmethod
-    def from_plugin_path(cls, plugin_path: typing.Union[pathlib.Path, str]):
+    def from_plugin_path(
+        cls,
+        plugin_path: Union[Path, str]
+    ) -> Optional[UnrealPlugin]:
         '''
         Get an UnrealPlugin object based on the given file system path.
 
@@ -23,11 +296,11 @@ class UnrealPlugin(object):
         a ".uplugin" file. In the latter case, the directory will be
         searched for the ".uplugin" file.
         '''
-        if not isinstance(plugin_path, pathlib.Path):
-            plugin_path = pathlib.Path(plugin_path)
+        if not isinstance(plugin_path, Path):
+            plugin_path = Path(plugin_path)
 
         if not plugin_path.exists():
-            LOGGER.warning(
+            logging.warning(
                 'Cannot find uplugin file or plugin directory at '
                 f'path: {plugin_path}')
             return None
@@ -37,12 +310,12 @@ class UnrealPlugin(object):
             return cls(plugin_path)
 
         if not plugin_path.is_dir():
-            LOGGER.warning(f'Plugin path is not a directory: {plugin_path}')
+            logging.warning(f'Plugin path is not a directory: {plugin_path}')
             return None
 
         uplugin_file_paths = list(plugin_path.glob('*.uplugin'))
         if not uplugin_file_paths:
-            LOGGER.warning(
+            logging.warning(
                 f'No .uplugin files found in plugin directory: {plugin_path}')
             return None
 
@@ -51,18 +324,19 @@ class UnrealPlugin(object):
                 'Multiple .uplugin files found in plugin directory: '
                 f'{plugin_path}\n    ')
             msg += '\n    '.join([str(p) for p in uplugin_file_paths])
-            LOGGER.warning(msg)
+            logging.warning(msg)
             return None
 
-        uplugin_file_path = pathlib.Path(uplugin_file_paths[0])
+        uplugin_file_path = Path(uplugin_file_paths[0])
 
         return cls(uplugin_file_path)
 
     @classmethod
     def from_path_filters(
-            cls,
-            ue_project_path: typing.Union[pathlib.Path, str],
-            filter_patterns: typing.List[str]):
+        cls,
+        uproject_path: Union[Path, str],
+        filter_patterns: list[str]
+    ) -> list[UnrealPlugin]:
         '''
         Get a list of Unreal Engine plugins matching the given path-based
         filter patterns.
@@ -74,13 +348,13 @@ class UnrealPlugin(object):
         relative path patterns are assumed to be relative project directory
         (the directory containing the .uproject file).
         '''
-        if not isinstance(ue_project_path, pathlib.Path):
-            ue_project_path = pathlib.Path(ue_project_path)
+        if not isinstance(uproject_path, Path):
+            uproject_path = Path(uproject_path)
 
         plugin_paths = set()
 
         for filter_pattern in filter_patterns:
-            path_pattern = pathlib.Path(filter_pattern)
+            path_pattern = Path(filter_pattern)
 
             if path_pattern.is_absolute():
                 plugin_filter = path_pattern
@@ -93,17 +367,18 @@ class UnrealPlugin(object):
                     # When the path pattern is relative and it looks like a
                     # directory name (i.e. no path separators), assume we're
                     # matching against plugin directories inside the project.
-                    plugin_filter = ue_project_path / 'Plugins' / path_pattern
+                    plugin_filter = (uproject_path.parent /
+                                     'Plugins' / path_pattern)
                 else:
                     # Otherwise, assume that the path pattern is relative to
                     # the project directory.
-                    plugin_filter = ue_project_path / path_pattern
+                    plugin_filter = uproject_path.parent / path_pattern
 
             for plugin_path in glob.glob(str(plugin_filter)):
-                plugin_path = pathlib.Path(plugin_path).resolve()
+                plugin_path = Path(plugin_path).resolve()
                 plugin_paths.add(plugin_path)
 
-        unreal_plugins = []
+        unreal_plugins: list[UnrealPlugin] = []
         for plugin_path in sorted(list(plugin_paths)):
             unreal_plugin = cls.from_plugin_path(plugin_path)
             if unreal_plugin:
@@ -111,17 +386,22 @@ class UnrealPlugin(object):
 
         return unreal_plugins
 
-    def __init__(self, uplugin_file_path: typing.Union[pathlib.Path, str]):
-        if not isinstance(uplugin_file_path, pathlib.Path):
-            uplugin_file_path = pathlib.Path(uplugin_file_path)
+    def __init__(
+        self,
+        uplugin_file_path: Union[Path, str],
+        plugin_type: UnrealPluginType = UnrealPluginType.Project,
+    ):
+        if not isinstance(uplugin_file_path, Path):
+            uplugin_file_path = Path(uplugin_file_path)
 
         self._uplugin_file_path = uplugin_file_path
+        self._plugin_type = plugin_type
 
     def __repr__(self) -> str:
         return f'{self.__class__.__name__}("{self._uplugin_file_path}")'
 
     @property
-    def uplugin_file_path(self) -> pathlib.Path:
+    def uplugin_file_path(self) -> Path:
         '''
         The file system path to the plugin's .uplugin file.
         '''
@@ -135,23 +415,23 @@ class UnrealPlugin(object):
         return self._uplugin_file_path.stem
 
     @property
-    def plugin_path(self) -> pathlib.Path:
+    def plugin_dir(self) -> Path:
         '''
         The file system path to the root directory of the plugin.
         '''
         return self._uplugin_file_path.parents[0]
 
     @property
-    def plugin_content_path(self) -> pathlib.Path:
+    def plugin_content_path(self) -> Path:
         '''
         The file system path to the plugin's "Content" directory.
 
         This is the directory that gets mounted in Unreal Engine.
         '''
-        return self.plugin_path / 'Content'
+        return self.plugin_dir / 'Content'
 
     @property
-    def mounted_path(self) -> pathlib.PurePosixPath:
+    def mounted_path(self) -> PurePosixPath:
         '''
         The root content path of the plugin when its "Content" directory is
         mounted in Unreal Engine.
@@ -159,4 +439,9 @@ class UnrealPlugin(object):
         # Note that UE uses the file name of the uplugin file to produce the
         # mounted path in engine, and not the name of the plugin directory or
         # any data inside the uplugin file.
-        return pathlib.PurePosixPath(f'/{self.name}')
+        return PurePosixPath(f'/{self.name}')
+
+    @property
+    def plugin_type(self) -> UnrealPluginType:
+        ''' The type of this plugin. '''
+        return self._plugin_type
diff --git a/Engine/Plugins/VirtualProduction/Switchboard/Source/Switchboard/switchboard/ugs_utils.py b/Engine/Plugins/VirtualProduction/Switchboard/Source/Switchboard/switchboard/ugs_utils.py
--- a/Engine/Plugins/VirtualProduction/Switchboard/Source/Switchboard/switchboard/ugs_utils.py
+++ b/Engine/Plugins/VirtualProduction/Switchboard/Source/Switchboard/switchboard/ugs_utils.py
@@ -533,6 +533,8 @@ def parse_depot_ugs_configs(
     for depot_path, record in zip(depot_print_paths, config_prints):
         if record.is_valid:
             config_contents_map[depot_path] = record.text
+        elif b'delete' in record.stat.get(b'action', b''):
+            del config_contents_map[depot_path]
         else:
             logging.warning(f'Invalid P4PrintResult for {depot_path}')
 
