Switchboard: Added optimizations to keep the UI responsive when connected to many (e.g. 30) nDisplay nodes.

* Changed to standard table model that seems more optimized and manages its own internal cache of table data.
* Spread nDisplay node status polls in time instead of doing all at the same time.
* Added indexed caches on write operations of the program queue so that multiple readers can find programs quickly.

diff --git a/Engine/Plugins/VirtualProduction/Switchboard/Source/Switchboard/switchboard/devices/ndisplay/ndisplay_monitor.py b/Engine/Plugins/VirtualProduction/Switchboard/Source/Switchboard/switchboard/devices/ndisplay/ndisplay_monitor.py
--- a/Engine/Plugins/VirtualProduction/Switchboard/Source/Switchboard/switchboard/devices/ndisplay/ndisplay_monitor.py
+++ b/Engine/Plugins/VirtualProduction/Switchboard/Source/Switchboard/switchboard/devices/ndisplay/ndisplay_monitor.py
@@ -1,13 +1,14 @@
 # Copyright Epic Games, Inc. All Rights Reserved.
 
 from collections import OrderedDict
-from itertools import count
+from itertools import count, islice
 import time
+import threading
 import traceback
 
 from PySide6 import QtCore
-from PySide6.QtCore import QAbstractTableModel, QModelIndex, Qt, QTimer
-from PySide6.QtGui import QColor, QIcon
+from PySide6.QtCore import QModelIndex, Qt, QTimer, Q_ARG
+from PySide6.QtGui import QColor, QIcon, QStandardItemModel, QStandardItem
 
 from switchboard import message_protocol
 from switchboard.message_protocol import SyncStatusRequestFlags
@@ -15,33 +16,39 @@ from switchboard.switchboard_logging import LOGGER
 from switchboard.devices.device_base import Device
 
 
-class nDisplayMonitor(QAbstractTableModel):
+class nDisplayMonitor(QStandardItemModel):
     '''
     This will monitor the status of the nDisplay nodes, in particular regarding
     sync. It polls the listener at the specified rate and the UI should update
     with this info.
     '''
 
-    COLOR_WARNING = QColor(0x70, 0x40, 0x00)
-    COLOR_NORMAL = QColor(0x3d, 0x3d, 0x3d)
-    COLOR_DISCONNECTED = QColor(0x7F, 0x7F, 0x7F)
+    BG_COLOR_WARNING = QColor(0x70, 0x40, 0x00)
+    BG_COLOR_NORMAL = QColor(0x3d, 0x3d, 0x3d)
+    FG_COLOR_DISCONNECTED = QColor(0x7F, 0x7F, 0x7F)
+    FG_COLOR_NORMAL = QColor(0xd8, 0xd8, 0xd8)
     CORE_OVERLOAD_THRESH = 90  # percent utilization
 
     # Special meaning when received from listener for PresentMode, but
     # sometimes also used for other values' display for consistency.
     DATA_MISSING = 'n/a'
 
-    # If false, the button to disable full screen optimizations is hidden
-    # and periodic polling of the state of this feature is disabled.
-    show_disable_fso_btn = False
+    def __init__(self, parent):
+        QStandardItemModel.__init__(self, parent)
 
-    # Determines if GPU stats are regularly queried or not.
-    poll_gpu_stats = False
+        # If false, the button to disable full screen optimizations is hidden
+        # and periodic polling of the state of this feature is disabled.
+        self.show_disable_fso_btn = False
 
-    def __init__(self, parent):
-        QAbstractTableModel.__init__(self, parent)
+        # Determines if GPU stats are regularly queried or not.
+        self.poll_gpu_stats = False
+
+        # The time it takes to round trip the cluster and poll the status of the same device
+        # Used to detect stale devices.
+        self.roundtrip_polling_period_ms = 0
 
-        self.polling_period_ms = 2000
+        # Keeps track of the round-robin polling of device sync status
+        self._next_roundrobin_poll_deviceIdx = 0
 
         # ordered so that we can map row indices to devices
         self.devicedatas = OrderedDict()
@@ -95,6 +102,7 @@ class nDisplayMonitor(QAbstractTableModel):
         }
 
         self.colnames = list(HEADER_DATA.keys())
+        self.colnamesdisplay = ["" if name == 'Connected' else name for name in self.colnames]
         self.tooltips = list(HEADER_DATA.values())
 
         # Load connection status icons
@@ -105,9 +113,7 @@ class nDisplayMonitor(QAbstractTableModel):
     def color_for_column(self, colname, value, data, is_program_running: bool):
         ''' Returns the background color for the given cell '''
         if data['Connected'].lower() == 'no':
-            if colname == 'Connected':
-                return self.COLOR_WARNING
-            return self.COLOR_NORMAL
+            return self.BG_COLOR_NORMAL
 
         if colname == 'PresentMode':
             ok_values = [
@@ -115,31 +121,31 @@ class nDisplayMonitor(QAbstractTableModel):
                 'Hardware: Independent Flip',
             ]
             is_good = (not is_program_running) or any(ok_value in value for ok_value in ok_values)
-            return self.COLOR_NORMAL if is_good else self.COLOR_WARNING
+            return self.BG_COLOR_NORMAL if is_good else self.BG_COLOR_WARNING
 
         if colname == 'Gpus':
             is_synced = ('Synced' in value) and ('Free' not in value)
-            return self.COLOR_NORMAL if is_synced else self.COLOR_WARNING
+            return self.BG_COLOR_NORMAL if is_synced else self.BG_COLOR_WARNING
 
         if colname == 'InFocus':
             is_good = (not is_program_running) or ('yes' in value)
-            return self.COLOR_NORMAL if is_good else self.COLOR_WARNING
+            return self.BG_COLOR_NORMAL if is_good else self.BG_COLOR_WARNING
 
         if colname == 'FSO':
             is_good = (not is_program_running) or ('no' in value)
-            return self.COLOR_NORMAL if is_good else self.COLOR_WARNING
+            return self.BG_COLOR_NORMAL if is_good else self.BG_COLOR_WARNING
 
         if colname == 'Displays':
             is_normal = (('Follower' in value or 'Leader' in value)
                          and ('Unsynced' not in value))
-            return self.COLOR_NORMAL if is_normal else self.COLOR_WARNING
+            return self.BG_COLOR_NORMAL if is_normal else self.BG_COLOR_WARNING
 
         if colname == 'CpuUtilization':
             # "(# cores > threshold%)" or "(SMT ENABLED)"
             no_caveats = '(' not in value
-            return self.COLOR_NORMAL if no_caveats else self.COLOR_WARNING
+            return self.BG_COLOR_NORMAL if no_caveats else self.BG_COLOR_WARNING
 
-        return self.COLOR_NORMAL
+        return self.BG_COLOR_NORMAL
 
     def friendly_osver(self, device):
         ''' Returns a display-friendly string for the device OS version '''
@@ -197,7 +203,7 @@ class nDisplayMonitor(QAbstractTableModel):
         data['Host'] = str(device.address)
         data['Node'] = device.name
         data['Connected'] = \
-            'yes' if device.unreal_client.is_connected else 'no'
+            'yes' if device.is_connected_and_authenticated() else 'no'
 
         # extra data not in columns
         data['TimeLastFlipGlitch'] = time.time()
@@ -213,22 +219,81 @@ class nDisplayMonitor(QAbstractTableModel):
         }
 
         # notify the UI of the change
-        self.layoutChanged.emit()
+        self.rebuild_table()
 
-        # start/continue polling since there is at least one device
-        if not self.timer.isActive():
-            self.timer.start(self.polling_period_ms)
+        # Adding a device may change our polling strategy
+        self.update_polling_timer()
+
+    def get_col_count(self) -> int:
+        ''' Returns the number of columns that the table should have.'''
+        return len(self.colnames)
+
+    def get_row_count(self) -> int:
+        ''' Returns the number of rows that the table should have.
+        It is based on the number of devices that have been added.
+        '''
+        return len(self.devicedatas)
+
+    def rebuild_table(self) -> None:
+        ''' Rebuilds all the table items from scratch. '''
+        self.setColumnCount(self.get_col_count())
+        self.setRowCount(self.get_row_count())
+        self.rebuild_headers()
+
+        # Update the row data
+        for row in range(self.get_row_count()):
+            self.refresh_display_for_row(row)
+
+    def rebuild_headers(self) -> None:
+        ''' Repopulates from scratch the header items.'''
+        # set header labels
+        self.setHorizontalHeaderLabels(self.colnamesdisplay)
+
+        # set tooltips
+        for col in range(self.get_col_count()):
+            self.setHeaderData(
+                col,
+                Qt.Orientation.Horizontal,
+                self.tooltips[col],
+                Qt.ItemDataRole.ToolTipRole
+            )
 
     def removed_device(self, device):
         ''' Called by the plugin when an nDisplay device has been removed. '''
         self.devicedatas.pop(device.device_hash)
 
-        # notify the UI of the change
-        self.layoutChanged.emit()
+        self.rebuild_table()  # To reflect this in the UI
+        self.update_polling_timer()  # This may change our polling strategy
 
-        # turn off the timer if there are no devices left
-        if not self.devicedatas:
+    def update_polling_timer(self) -> None:
+        ''' Re-calculates what the polling timer should be set to.'''
+
+        devs = [devicedata['device'] for devicedata in self.devicedatas.values()]
+
+        # No need for a timer if there aren't any devices in the table
+        if len(devs) == 0:
             self.timer.stop()
+            return
+
+        num_pollable_nodes = sum(1 for dev in devs if self.can_poll_device(dev))
+
+        # turn off the timer if there are no devices left
+        if num_pollable_nodes:
+            min_ms_between_nodes = 250  # Fastest allowed update rate per node
+            min_ms_roundtrip = 1000     # Fastest allowed update rate of the same node
+
+            polling_period_ms = max(min_ms_roundtrip / num_pollable_nodes, min_ms_between_nodes)
+
+            self.roundtrip_polling_period_ms = num_pollable_nodes * polling_period_ms
+
+            # start/continue polling since there is at least one device
+            self.timer.setInterval(polling_period_ms)
+        else:
+            maintenance_timer_tick_ms = 2000
+            self.timer.setInterval(maintenance_timer_tick_ms)
+
+        if not self.timer.isActive():
+            self.timer.start()
 
     def handle_stale_device(self, devicedata, deviceIdx):
         '''
@@ -244,7 +309,7 @@ class nDisplayMonitor(QAbstractTableModel):
             time.time() - devicedata['time_last_update']
 
         timeout_factor = 4
-        timeout = timeout_factor * self.polling_period_ms * 1e-3
+        timeout = timeout_factor * self.roundtrip_polling_period_ms * 1e-3
         if time_elapsed_since_last_update < timeout:
             return
 
@@ -253,7 +318,7 @@ class nDisplayMonitor(QAbstractTableModel):
         devicedata['stale'] = True
 
         # notify the UI
-        row = deviceIdx + 1
+        row = deviceIdx
         self.refresh_display_for_row(row)
 
     def handle_connection_change(self, devicedata, deviceIdx):
@@ -264,18 +329,24 @@ class nDisplayMonitor(QAbstractTableModel):
         device = devicedata['device']
         data = devicedata['data']
 
-        is_connected = device.unreal_client.is_connected
+        is_connected = device.is_connected_and_authenticated()
         was_connected = True if data['Connected'] == 'yes' else False
 
         data['Connected'] = 'yes' if is_connected else 'no'
 
         if was_connected != is_connected:
-            if not is_connected:
+            if is_connected:
+                # Poll the sync status to get an immediate update on its state
+                self.poll_sync_status_for_device(device, SyncStatusRequestFlags.all())
+            else:
                 self.reset_device_data(device, data)
 
-            row = deviceIdx + 1
+            row = deviceIdx
             self.refresh_display_for_row(row)
 
+            # Since we only poll connected devices, our polling timer needs to be updated.
+            self.update_polling_timer()
+
     def default_program_id(self):
         ''' Default value for program id when unreal is not running.
         Used for messages that need this argument'''
@@ -291,12 +362,55 @@ class nDisplayMonitor(QAbstractTableModel):
 
         return program_id
 
+    def can_poll_device(self, device: Device) -> bool:
+        ''' Returns True if polling this device is allowed. Currently,
+        this is based on whether the listener is ready to take commands.'''
+        # can't poll if not connected to listener
+        return device.is_connected_and_authenticated()
+
+    @QtCore.Slot(Device)
+    def on_device_connected(self, device: Device) -> None:
+        ''' Called when a node has connected to its listener. '''
+        # Ensure this code is run from the main thread
+        if threading.current_thread() is not threading.main_thread():
+            QtCore.QMetaObject.invokeMethod(
+                self,
+                'on_device_connected',
+                QtCore.Qt.QueuedConnection,
+                Q_ARG(Device, device)
+            )
+            return
+
+        try:
+            deviceIdx, devicedata = self.devicedata_from_device(device)
+            self.handle_connection_change(devicedata=devicedata, deviceIdx=deviceIdx)
+        except KeyError:
+            LOGGER.warning(f"nDisplay Monitor could not find find device {device.name} upon connection")
+
+    @QtCore.Slot(Device)
+    def on_device_disconnected(self, device: Device) -> None:
+        ''' Called when a node has disconnected from its listener. '''
+        # Ensure this code is run from the main thread
+        if threading.current_thread() is not threading.main_thread():
+            QtCore.QMetaObject.invokeMethod(
+                self,
+                'on_device_disconnected',
+                QtCore.Qt.QueuedConnection,
+                Q_ARG(Device, device)
+            )
+            return
+
+        try:
+            deviceIdx, devicedata = self.devicedata_from_device(device)
+            self.handle_connection_change(devicedata=devicedata, deviceIdx=deviceIdx)
+        except KeyError:
+            LOGGER.warning(f"nDisplay Monitor could not find find device {device.name} upon disconnection")
+
     def poll_sync_status_for_device(self, device: Device, request_flags: SyncStatusRequestFlags):
         ''' Polls sync status for the given device '''
 
         # no point in continuing if not connected to listener
-        if not (device.unreal_client.is_connected
-                and device.unreal_client.is_authenticated):
+        if not self.can_poll_device(device):
             return
 
         # create message
@@ -309,22 +423,38 @@ class nDisplayMonitor(QAbstractTableModel):
         # send get sync status message
         device.unreal_client.send_message(msg)
 
-    def poll_sync_status(self, request_flags: SyncStatusRequestFlags):
-        ''' Polls sync status for all nDisplay devices '''
+    def poll_sync_status(self, request_flags: SyncStatusRequestFlags, all=False) -> None:
+        ''' Polls sync status for all nDisplay devices in round robin '''
 
-        for deviceIdx, devicedata in enumerate(self.devicedatas.values()):
+        numdevs = len(self.devicedatas)
+
+        for _ in range(numdevs):
+            if self._next_roundrobin_poll_deviceIdx >= numdevs:
+                self._next_roundrobin_poll_deviceIdx = 0
+
+            # Get the device we're interested in polling and advance the round robin index
+            deviceIdx = self._next_roundrobin_poll_deviceIdx
+            devicedata = next(islice(self.devicedatas.values(), deviceIdx, deviceIdx + 1))
+            self._next_roundrobin_poll_deviceIdx += 1
             device = devicedata['device']
 
+            # We still need to do some maintenance on all the devices. By iterating
+            # over them even in sections per polling timout we keep them up to date.
+
             # detect connection changes (a disconnection invalidates data)
             self.handle_connection_change(devicedata, deviceIdx)
 
             # detect stale devices
             self.handle_stale_device(devicedata, deviceIdx)
 
-            # request status
-            self.poll_sync_status_for_device(device, request_flags)
+            # Poll the device if we can.
+            if self.can_poll_device(device):
+                # request status
+                self.poll_sync_status_for_device(device, request_flags)
+                if not all:
+                    return
 
-    def poll_variable_sync_status(self):
+    def poll_variable_sync_status(self) -> None:
         ''' Poll sync status but only include items that may change and that are
         not likely to cause hitches in the target machine. '''
 
@@ -619,7 +749,7 @@ class nDisplayMonitor(QAbstractTableModel):
                 f'{traceback.format_exc()}=== Traceback END ===\n')
             return
 
-        row = deviceIdx + 1
+        row = deviceIdx
         self.refresh_display_for_row(row)
 
     def try_issue_console_exec(self, exec_str, executor=''):
@@ -637,90 +767,74 @@ class nDisplayMonitor(QAbstractTableModel):
 
         return False
 
+    def make_default_table_item(self):
+        ''' Creates a standard table item with default configuration '''
+        item = QStandardItem()
+        item.setFlags(item.flags() & ~QtCore.Qt.ItemIsEditable & ~QtCore.Qt.ItemIsSelectable)
+        return item
+
     def refresh_display_for_row(self, row: int):
         ''' Refreshes the data displayed in the given row. Zero-based index. '''
+        # Gather the data
+        _, devicedata = list(self.devicedatas.items())[row]
+        device = devicedata['device']
+        data = devicedata['data']
+        is_program_running = self.program_id_from_device(device) != self.default_program_id()
 
-        self.refresh_display_for_cells(
-            self.createIndex(row, 0),
-            self.createIndex(row, len(self.colnames) - 1)
-        )
+        if data['Connected'].lower() == 'yes':
+            is_connected = True
+        else:
+            is_connected = False
 
-    def refresh_display_for_cells(self, upper_left: QModelIndex, lower_right: QModelIndex):
-        ''' Refreshes the cells in the given range, inclusive '''
+        # Update the item
+        for col, colname in enumerate(self.colnames):
 
-        self.dataChanged.emit(upper_left, lower_right)
+            item = self.item(row, col)
 
-        # This was necessary to reliably update the table display with the changed data.
-        self.headerDataChanged.emit(Qt.Orientation.Horizontal, upper_left.column, lower_right.column)
+            if not item:
+                item = self.make_default_table_item()
+                self.setItem(row, col, item)
 
-    # ~ QAbstractTableModel interface begin
+            # Set text
+            value = data[colname]
+            item.setText(value)
 
-    def rowCount(self, parent=QModelIndex()):
-        return len(self.devicedatas)
+            # Set background color
+            bg_color = self.color_for_column(
+                colname=colname,
+                value=value,
+                data=data,
+                is_program_running=is_program_running
+            )
 
-    def columnCount(self, parent=QModelIndex()):
-        return len(self.colnames)
+            item.setBackground(bg_color)
 
-    def headerData(self, section, orientation, role):
-        if role == Qt.ItemDataRole.DisplayRole:
-            if orientation == Qt.Orientation.Horizontal:
-                # Connection column uses icons and don't need the long header name
-                if self.colnames[section] == 'Connected':
-                    return ""
-                return self.colnames[section]
+            # Set foreground color
+            if is_connected:
+                item.setForeground(self.FG_COLOR_NORMAL)
             else:
-                return "{}".format(section)
-
-        if role == Qt.ItemDataRole.ToolTipRole:
-            return self.tooltips[section]
-
-        return None
-
-    def data(self, index, role=Qt.ItemDataRole.DisplayRole):
-        column = index.column()
-        row = index.row()
-
-        # get column name
-        colname = self.colnames[column]
+                item.setForeground(self.FG_COLOR_DISCONNECTED)
 
-        # grab (device_hash, device_data) from ordered dict
-        _, devicedata = list(self.devicedatas.items())[row]
-        device = devicedata['device']
-        data = devicedata['data']
-        is_program_running = self.program_id_from_device(device) != self.default_program_id()
+            # Set the connection icon
+            if colname == 'Connected':
+                icon = self.icon_unconnected
 
-        if colname == 'Connected':
-            if role == Qt.ItemDataRole.DecorationRole:
-                if data['Connected'].lower() == 'yes':
+                if value.lower() == 'yes':
                     if is_program_running:
-                        return self.icon_running
-                    return self.icon_connected
-                return self.icon_unconnected
-
-            return None
-
-        value = data[colname]
-
-        if role == Qt.ItemDataRole.DisplayRole:
-            return value
+                        icon = self.icon_running
+                    else:
+                        icon = self.icon_connected
 
-        elif role == Qt.ItemDataRole.BackgroundRole:
-            return self.color_for_column(colname=colname, value=value,
-                                         data=data, is_program_running=is_program_running)
+                item.setIcon(icon)
+                item.setText("")
 
-        elif role == Qt.ItemDataRole.TextAlignmentRole:
-            alignment = Qt.AlignmentFlag.AlignCenter
+            # Set the alignment
             if colname in ('CpuUtilization', 'GpuUtilization'):
                 alignment = Qt.AlignmentFlag.AlignLeft | Qt.AlignmentFlag.AlignVCenter
-            return alignment
-
-        elif role == Qt.ItemDataRole.ForegroundRole:
-            if data['Connected'].lower() != 'yes':
-                return self.COLOR_DISCONNECTED
-
-        return None
+            else:
+                alignment = Qt.AlignmentFlag.AlignCenter
 
-    # ~ QAbstractTableModel interface end
+            item.setTextAlignment(alignment)
 
     def update_device_name_and_address(self, in_device: Device):
         ''' Refeshes the node names and addresses displayed in the table '''
@@ -737,7 +851,7 @@ class nDisplayMonitor(QAbstractTableModel):
                 data['Node'] = device.name
 
                 # Refresh the table with the new data
-                row = deviceIdx + 1
+                row = deviceIdx
                 self.refresh_display_for_row(row)
                 return
 
@@ -752,7 +866,7 @@ class nDisplayMonitor(QAbstractTableModel):
             self.update_device_name_and_address(device)
 
         # Request a full update on the sync status.
-        self.poll_sync_status(SyncStatusRequestFlags.all())
+        self.poll_sync_status(SyncStatusRequestFlags.all(), all=True)
 
     @QtCore.Slot()
     def on_disable_fso_clicked(self):
diff --git a/Engine/Plugins/VirtualProduction/Switchboard/Source/Switchboard/switchboard/devices/ndisplay/plugin_ndisplay.py b/Engine/Plugins/VirtualProduction/Switchboard/Source/Switchboard/switchboard/devices/ndisplay/plugin_ndisplay.py
--- a/Engine/Plugins/VirtualProduction/Switchboard/Source/Switchboard/switchboard/devices/ndisplay/plugin_ndisplay.py
+++ b/Engine/Plugins/VirtualProduction/Switchboard/Source/Switchboard/switchboard/devices/ndisplay/plugin_ndisplay.py
@@ -1726,13 +1726,17 @@ class DevicenDisplay(DeviceUnreal):
 
     # ~ Begin DeviceUnreal Interface
     @QtCore.Slot()
-    def _on_listener_connected(self):
+    def _on_listener_connected(self) -> None:
         super()._on_listener_connected()
         if self.__class__.ndisplay_monitor:
-            # Poll the sync status to get an immediate update on its state
-            self.__class__.ndisplay_monitor.poll_sync_status_for_device(self, SyncStatusRequestFlags.all())
+            self.__class__.ndisplay_monitor.on_device_connected(device=self)
 
-    def on_unreal_started(self):
+    def on_listener_disconnect(self, unexpected=False, exception=None) -> None:
+        super().on_listener_disconnect(unexpected, exception)
+        if self.__class__.ndisplay_monitor:
+            self.__class__.ndisplay_monitor.on_device_disconnected(device=self)
+
+    def on_unreal_started(self) -> None:
         super().on_unreal_started()
 
         # Always try to disable Full Screen Optimizations. If this is removed, you can add back
diff --git a/Engine/Plugins/VirtualProduction/Switchboard/Source/Switchboard/switchboard/devices/unreal/plugin_unreal.py b/Engine/Plugins/VirtualProduction/Switchboard/Source/Switchboard/switchboard/devices/unreal/plugin_unreal.py
--- a/Engine/Plugins/VirtualProduction/Switchboard/Source/Switchboard/switchboard/devices/unreal/plugin_unreal.py
+++ b/Engine/Plugins/VirtualProduction/Switchboard/Source/Switchboard/switchboard/devices/unreal/plugin_unreal.py
@@ -84,9 +84,18 @@
     '''
     @wraps(func)
     def _use_lock(self, *args, **kwargs):
+
         self.lock.acquire()
+
         try:
-            return func(self, *args, **kwargs)
+            ret = func(self, *args, **kwargs)
+
+            # Using the lock implies a write operation, which requires
+            # updating the reader indexed caches.
+            self._update_rd_caches()
+
+            return ret
+
         finally:
             self.lock.release()
 
@@ -94,7 +103,16 @@
 
 
 class ProgramStartQueue:
-    ''' Queue of programs to launch that may have dependencies '''
+    ''' Queue of programs to launch that may have dependencies 
+    Write operations can from from the listener thread (e.g. program events)
+    or from the Main thread (e.g. add a program). Read operations can happen
+    on the Main thread (e.g. see if a program is running) or from the listener
+    thread.
 
+    The locking strategy is for the writers to use a lock and update a shallow
+    copy for the readers. Reading is expected to happen much more often than writing,
+    especially in relation to checking if programs are running or not.
+    '''
+
     def __init__(self):
         # Needed because these functions will be called from listener thread
@@ -102,11 +120,16 @@
         self.lock = threading.Lock()
 
         self.queued_programs: list[ProgramStartQueueItem] = []
-        self.starting_programs: OrderedDict[
-            uuid.UUID, ProgramStartQueueItem] = OrderedDict()
+        self.starting_programs: OrderedDict[uuid.UUID, ProgramStartQueueItem] = OrderedDict()
+
         # initialized as returned by listener
-        self.running_programs: OrderedDict[
-            uuid.UUID, ProgramStartQueueItem] = OrderedDict()
+        self.running_programs: OrderedDict[uuid.UUID, ProgramStartQueueItem] = OrderedDict()
 
+        # Indexed caches for the readers.
+        self.rd_running_progs_by_name: dict[str, list[ProgramStartQueueItem]] = dict()
+        self.rd_puuid_by_name: dict[uuid.UUID, str] = dict()
+
+    # Write operations
+
     @use_lock
     def reset(self):
@@ -115,39 +138,6 @@
         self.starting_programs = OrderedDict()
         self.running_programs = OrderedDict()
 
-    def _name_from_puuid(self, puuid):
-        ''' Returns the name of the specified program id '''
-        for prog in self.queued_programs:
-            if prog.puuid == puuid:
-                return prog.name
-
-        for thedict in [self.starting_programs, self.running_programs]:
-            try:
-                return thedict[puuid].name
-            except KeyError:
-                pass
-
-        raise KeyError
-
     @use_lock
-    def puuid_from_name(self, name):
-        '''
-        Returns the puuid of the specified program name.
-
-        Searches as to return the one most likely to return last - but not
-        guaranteed to do so.
-        '''
-        for prog in self.queued_programs:
-            if prog.name == name:
-                return prog.puuid
-
-        for thedict in [self.starting_programs, self.running_programs]:
-            for prog in thedict.values():
-                if prog.name == name:
-                    return prog.puuid
-
-        raise KeyError
-
-    @use_lock
     def on_program_started(self, prog):
         ''' Returns the name of the program that started
@@ -204,6 +194,40 @@
 
         return prog.name if prog else "unknown"
 
+    @use_lock
+    def add(self, prog, unreal_client) -> None:
+        ''' Adds a new program to be started in the queue
+        Must be of type ProgramStartQueueItem.
+        It may start launch it right away if it doesn't have any dependencies.
+        '''
+        assert isinstance(prog, ProgramStartQueueItem)
+
+        # Ensure that dependence is still possible, and if it isn't, replace
+        # with None.
+        if prog.puuid_dependency is not None:
+            try:
+                self._name_from_puuid(prog.puuid_dependency)
+            except KeyError:
+                LOGGER.debug(
+                    f"{prog.name} specified non-existent dependency on puuid "
+                    f"{prog.puuid_dependency}")
+                prog.puuid_dependency = None
+
+        self.queued_programs.append(prog)
+
+        # This effectively causes a launch if it doesn't have any dependencies
+        self._launch_dependents(puuid=None, unreal_client=unreal_client)
+
+    @use_lock
+    def clear_running_programs(self) -> None:
+        self.running_programs.clear()
+
+    @use_lock
+    def update_running_program(self, prog) -> None:
+        self.running_programs[prog.puuid] = prog
+
+    # Helper functions
+
     def _launch_dependents(self, puuid, unreal_client):
         ''' Launches programs dependend on given puuid
         Do not call externally because it does not use the thread lock.
@@ -223,51 +247,78 @@
         for prog in progs_launched:
             self.queued_programs.remove(prog)
 
-    @use_lock
-    def add(self, prog, unreal_client):
-        ''' Adds a new program to be started in the queue
-        Must be of type ProgramStartQueueItem.
-        It may start launch it right away if it doesn't have any dependencies.
-        '''
-        assert isinstance(prog, ProgramStartQueueItem)
+    def _name_from_puuid(self, puuid) -> str:
+        ''' Returns the name of the specified program id '''
+        for prog in self.queued_programs:
+            if prog.puuid == puuid:
+                return prog.name
 
-        # Ensure that dependance is still possible, and if it isn't, replace
-        # with None.
-        if prog.puuid_dependency is not None:
+        for thedict in [self.starting_programs, self.running_programs]:
             try:
-                self._name_from_puuid(prog.puuid_dependency)
+                return thedict[puuid].name
             except KeyError:
-                LOGGER.debug(
-                    f"{prog.name} specified non-existent dependency on puuid "
-                    f"{prog.puuid_dependency}")
-                prog.puuid_dependency = None
+                pass
 
-        self.queued_programs.append(prog)
+        raise KeyError
 
-        # This effectively causes a launch if it doesn't have any dependencies
-        self._launch_dependents(puuid=None, unreal_client=unreal_client)
+    def _update_rd_caches(self) -> None:
+        ''' Updates the indexed dictionaries for the readers.
+        This function assumes it is called by one of the writer functions
+        that already hold the lock.
+        '''
 
-    @use_lock
-    def running_programs_named(self, name):
+        # Update rd_running_progs_by_name
+        #
+        rd_running_progs_by_name: dict[str, list[ProgramStartQueueItem]] = dict()
+
+        for prog in self.running_programs.values():
+
+            if prog.name not in rd_running_progs_by_name:
+                rd_running_progs_by_name[prog.name] = []
+
+            rd_running_progs_by_name[prog.name].append(prog)
+
+        # Update rd_puuid_by_name
+        #
+        rd_puuid_by_name: dict[str, uuid.UUID] = dict()
+
+        # Combine all programs into one list, such that when iterated,
+        # the last program will be the the one more likely to finish last.
+        # I.e. a queued program of a given name is expected to end after
+        # one that is already running.
+        progs = reversed(
+            list(self.queued_programs) +
+            list(self.starting_programs.values()) +
+            list(self.running_programs.values())
+        )
+
+        for prog in progs:
+            rd_puuid_by_name[prog.name] = prog.puuid
+
+        # Atomically replace the dicts for the readers.
+
+        self.rd_running_progs_by_name = rd_running_progs_by_name
+        self.rd_puuid_by_name = rd_puuid_by_name
+
+    # Read operations
+
+    def running_programs_named(self, name) -> list[ProgramStartQueueItem]:
         ''' Returns the program ids of running programs named as specified '''
-        return [
-            prog for puuid, prog in self.running_programs.items()
-            if prog.name == name]
+        return self.rd_running_progs_by_name.get(name, [])
 
-    @use_lock
-    def running_puuids_named(self, name):
+    def running_puuids_named(self, name) -> list[uuid.UUID]:
         ''' Returns the program ids of running programs named as specified '''
-        return [
-            puuid for puuid, prog in self.running_programs.items()
-            if prog.name == name]
+        progs = self.running_programs_named(name)
+        return [prog.puuid for prog in progs]
 
-    @use_lock
-    def update_running_program(self, prog):
-        self.running_programs[prog.puuid] = prog
+    def puuid_from_name(self, name) -> uuid.UUID:
+        '''
+        Returns the puuid of the specified program name.
 
-    @use_lock
-    def clear_running_programs(self):
-        self.running_programs.clear()
+        Searches as to return the one most likely to return last - but not
+        guaranteed to do so.
+        '''
+        return self.rd_puuid_by_name[name]
 
 
 class LiveLinkPresetSetting(Setting):
@@ -848,6 +899,13 @@
     def is_designated_local_builder(self) -> bool:
         return self is DeviceUnreal.get_designated_local_builder()
 
+    def is_connected_and_authenticated(self) -> bool:
+        ''' Returns true if it is connected to the listener and authenticated.
+        If true, it means that the application layer can have normal communiations with
+        the device
+        '''
+        return self.unreal_client.is_connected and self.unreal_client.is_authenticated
+
     @QtCore.Slot()
     def _queue_notify_redeploy(self):
         # Ensure this code is run from the main thread
