Switchboard: Improved handling of host key fingerprint mismatch.

When encountering an unexpected fingerprint connecting to a device, a warning dialog is now shown, explaining the situation and allowing the user to proceed if desired.

In order to help elucidate potential causes of key mismatch, the subject common name of the self-signed certificate now includes the OS username and computer name, and this information is surfaced in the warning dialog.

diff --git a/Engine/Plugins/VirtualProduction/Switchboard/Source/Switchboard/switchboard/__main__.py b/Engine/Plugins/VirtualProduction/Switchboard/Source/Switchboard/switchboard/__main__.py
--- a/Engine/Plugins/VirtualProduction/Switchboard/Source/Switchboard/switchboard/__main__.py
+++ b/Engine/Plugins/VirtualProduction/Switchboard/Source/Switchboard/switchboard/__main__.py
@@ -13,7 +13,7 @@ from .switchboard_logging import LOGGER
 
 
 # Build resources
-# "D:\P4\UE5\Main\Engine\Extras\ThirdPartyNotUE\SwitchboardThirdParty\Python\Scripts\pyside6-rcc" -o "D:\P4\UE5\Main\Engine\Plugins\VirtualProduction\Switchboard\Source\Switchboard\switchboard\resources.py" "D:\P4\UE5\Main\Engine\Plugins\VirtualProduction\Switchboard\Source\Switchboard\switchboard\ui\resources.qrc"
+# "Engine\Extras\ThirdPartyNotUE\SwitchboardThirdParty\Python\Scripts\pyside6-rcc" -o "Engine\Plugins\VirtualProduction\Switchboard\Source\Switchboard\switchboard\resources.py" "Engine\Plugins\VirtualProduction\Switchboard\Source\Switchboard\switchboard\ui\resources.qrc"
 
 
 def parse_arguments():
diff --git a/Engine/Plugins/VirtualProduction/Switchboard/Source/Switchboard/switchboard/credential_store.py b/Engine/Plugins/VirtualProduction/Switchboard/Source/Switchboard/switchboard/credential_store.py
--- a/Engine/Plugins/VirtualProduction/Switchboard/Source/Switchboard/switchboard/credential_store.py
+++ b/Engine/Plugins/VirtualProduction/Switchboard/Source/Switchboard/switchboard/credential_store.py
@@ -2,7 +2,7 @@
 
 from __future__ import annotations
 
-from abc import abstractmethod, ABC, ABCMeta
+from abc import abstractmethod, ABC
 from dataclasses import dataclass
 import json
 import logging
@@ -11,33 +11,15 @@ import pathlib
 import stat
 import sys
 import tempfile
-import threading
-from typing import Callable, Optional
+from typing import Optional
 
-from PySide6 import QtCore, QtWidgets
 
-
-QObjectMeta = type(QtCore.QObject)
-
-
-class QObjectABCMeta(QObjectMeta, ABCMeta):
-    pass
-
-
-class CredentialStore(QtCore.QObject, ABC, metaclass=QObjectABCMeta):
+class CredentialStore(ABC):
     @dataclass
     class Credential:
         username: str
         blob: str
 
-    # def OnCredentialReady(credential: Optional[Credential], unsaved: bool)
-    ReadyCallback = Callable[
-        [Optional[Credential], bool],
-        None]
-
-    def __init__(self):
-        super().__init__()
-
     @classmethod
     def encrypted_at_rest(cls) -> bool:
         return False
@@ -45,11 +27,7 @@ class CredentialStore(QtCore.QObject, ABC, metaclass=QObjectABCMeta):
     @abstractmethod
     def get(
         self,
-        key: str,
-        *,
-        context: Optional[str] = None,
-        context_long: Optional[str] = None,
-        on_credential_ready: Optional[ReadyCallback] = None
+        key: str
     ) -> Optional[Credential]:
         pass
 
@@ -57,7 +35,7 @@ class CredentialStore(QtCore.QObject, ABC, metaclass=QObjectABCMeta):
     def set(
         self,
         key: str,
-        val: Optional[CredentialStore.Credential]
+        val: Optional[Credential]
     ) -> None:
         pass
 
@@ -72,41 +50,15 @@ class CredentialStore(QtCore.QObject, ABC, metaclass=QObjectABCMeta):
 
 
 class CredentialStoreGeneric(CredentialStore):
-    @dataclass
-    class QueuedPrompt:
-        key: str
-        context: Optional[str]
-        context_long: Optional[str]
-        on_credential_ready: Optional[CredentialStore.ReadyCallback] = None
-
     def __init__(self):
         super().__init__()
 
         self._tempstore: dict[str, CredentialStore.Credential] = {}
-        self._queued_prompts: list[CredentialStoreGeneric.QueuedPrompt] = []
 
     def get(
         self,
-        key: str,
-        *,
-        context: Optional[str] = None,
-        context_long: Optional[str] = None,
-        on_credential_ready: Optional[CredentialStore.ReadyCallback] = None
+        key: str
     ) -> Optional[CredentialStore.Credential]:
-        if cred := self._try_get(key):
-            if on_credential_ready:
-                on_credential_ready(cred, False)
-            return cred
-
-        self._queued_prompts.append(
-            CredentialStoreGeneric.QueuedPrompt(
-                key=key, context=context, context_long=context_long,
-                on_credential_ready=on_credential_ready)
-        )
-
-        self._show_queued_prompts_soon()
-
-    def _try_get(self, key: str) -> Optional[CredentialStore.Credential]:
         return self._tempstore.get(key)
 
     def set(
@@ -119,30 +71,6 @@ class CredentialStoreGeneric(CredentialStore):
         elif key in self._tempstore:
             self._tempstore.pop(key)
 
-    @QtCore.Slot()
-    def _show_queued_prompts_soon(self):
-        # Ensure this code is run from the main thread
-        if threading.current_thread() is not threading.main_thread():
-            QtCore.QMetaObject.invokeMethod(
-                self, '_show_queued_prompts_soon', QtCore.Qt.QueuedConnection)
-            return
-
-        for prompt in self._queued_prompts:
-            blob, ok = QtWidgets.QInputDialog.getText(
-                None,
-                prompt.context or '',
-                prompt.context_long or '',
-                QtWidgets.QLineEdit.EchoMode.Password)
-
-            credential: Optional[CredentialStore.Credential] = None
-            if blob and ok:
-                credential = CredentialStore.Credential('', blob)
-
-            if prompt.on_credential_ready:
-                prompt.on_credential_ready(credential, True)
-
-        self._queued_prompts.clear()
-
 
 class CredentialStorePosix(CredentialStoreGeneric):
     PERMISSION_MASK = stat.S_IRWXU | stat.S_IRWXG | stat.S_IRWXO
@@ -162,7 +90,7 @@ class CredentialStorePosix(CredentialStoreGeneric):
     def _credential_path(cls, key: str) -> pathlib.Path:
         return pathlib.Path(cls._credential_dir(), f'credential_{key}.json')
 
-    def _try_get(self, key: str) -> Optional[CredentialStore.Credential]:
+    def get(self, key: str) -> Optional[CredentialStore.Credential]:
         path = self._credential_path(key)
         if not path or not path.is_file():
             return None
@@ -197,7 +125,7 @@ class CredentialStorePosix(CredentialStoreGeneric):
 
         (tmp_fd, tmp_path) = tempfile.mkstemp(dir=cred_dir, text=True)
         with os.fdopen(tmp_fd, mode='w') as f:
-            f.write(json.dumps({ 'user': val.username, 'blob': val.blob }))
+            f.write(json.dumps({'user': val.username, 'blob': val.blob}))
         os.replace(tmp_path, cred_path)
 
 
@@ -316,7 +244,7 @@ if sys.platform.startswith('win'):
         def encrypted_at_rest(cls) -> bool:
             return True
 
-        def _try_get(self, key: str) -> Optional[CredentialStore.Credential]:
+        def get(self, key: str) -> Optional[CredentialStore.Credential]:
             need_credfree = False
             cred = ctypes.pointer(CREDENTIALW())
             try:
diff --git a/Engine/Plugins/VirtualProduction/Switchboard/Source/Switchboard/switchboard/listener_client.py b/Engine/Plugins/VirtualProduction/Switchboard/Source/Switchboard/switchboard/listener_client.py
--- a/Engine/Plugins/VirtualProduction/Switchboard/Source/Switchboard/switchboard/listener_client.py
+++ b/Engine/Plugins/VirtualProduction/Switchboard/Source/Switchboard/switchboard/listener_client.py
@@ -2,9 +2,12 @@
 
 from __future__ import annotations
 
+from abc import abstractmethod, ABC
 import asyncio
+from asyncio import Future
 from collections import deque
 from datetime import datetime, timedelta
+import logging
 import ssl
 import threading
 import traceback
@@ -16,7 +19,9 @@ from aioquic.asyncio.protocol import QuicConnectionProtocol
 from aioquic.quic.configuration import QuicConfiguration
 from aioquic.quic.events import QuicEvent, HandshakeCompleted
 from cryptography.hazmat.primitives import hashes
-from PySide6 import QtCore
+from cryptography.x509.oid import NameOID
+from PySide6 import QtCore, QtGui, QtWidgets
+from PySide6.QtWidgets import QMessageBox
 
 from . import message_protocol
 from .credential_store import CredentialStore, CREDENTIAL_STORE
@@ -32,6 +37,7 @@ class SwitchboardClientProtocol(QuicConnectionProtocol):
 
         self.client: Optional[ListenerClient] = None
         self.fingerprint_sha2: Optional[bytes] = None
+        self.subject_cn = ''
 
     # @override
     def quic_event_received(self, event: QuicEvent) -> None:
@@ -44,71 +50,60 @@ class SwitchboardClientProtocol(QuicConnectionProtocol):
 
             self.fingerprint_sha2 = server_cert.fingerprint(hashes.SHA256())
 
+            subject = server_cert.subject
+            cn_attrs = subject.get_attributes_for_oid(NameOID.COMMON_NAME)
+            self.subject_cn = '+'.join([str(attr.value) for attr in cn_attrs])
+
     async def authenticate(self) -> bool:
+        assert self.client
         assert self.fingerprint_sha2
+
         fingerprint_str = self.fingerprint_sha2.hex()
 
         # truncated to 128 bits for visual comparison
         fingerprint_display_str = self.fingerprint_sha2[:16].hex(':').upper()
 
         host = self.client.address
+        credential = CREDENTIAL_STORE.get(key=host)
+        unsaved = credential is None
 
-        auth_response = self._loop.create_future()
-        closure_credential: Optional[CredentialStore.Credential] = None
-        closure_cred_unsaved = False
-
-        def on_credential_ready(
-            cred: Optional[CredentialStore.Credential],
-            unsaved: bool
-        ):
-            def auth_exception():
-                auth_response.set_exception(RuntimeError)
-                self.close()
-
-            if not cred:
-                self._quic._logger.error(
-                    'No credential for %s', host)
-                self._loop.call_soon_threadsafe(auth_exception)
-                return
+        if not credential:
+            password = await self.client.listener_qt_handler.prompt_password(
+                host, self.subject_cn, fingerprint_display_str)
 
-            nonlocal closure_credential, closure_cred_unsaved
-            closure_credential = cred
-            closure_cred_unsaved = unsaved
+            if password:
+                credential = CredentialStore.Credential('', password)
+            else:
+                logging.error('No credential for %s', host)
+                return False
 
+        if not credential.username:
             # First time update with expected remote fingerprint
-            if not cred.username:
-                cred.username = fingerprint_str
-            elif cred.username != fingerprint_str:
-                self._quic._logger.error(
-                    'MISMATCHED FINGERPRINT for %s\n\n'
-                    '%s (expected)\n'
-                    '%s (actual)',
-                    host, cred.username, fingerprint_str)
-                self._loop.call_soon_threadsafe(auth_exception)
-                return
-
-            if self.client:
-                jwt: Optional[str] = None
-                password: Optional[str] = None
-                if not unsaved and not CREDENTIAL_STORE.encrypted_at_rest():
-                    jwt = cred.blob
-                else:
-                    password = cred.blob
-
-                msgid, msg = message_protocol.create_authenticate_message(
-                    jwt=jwt, password=password)
+            credential.username = fingerprint_str
+        elif credential.username != fingerprint_str:
+            # Warn the user if the key changed
+            result = await self.client.listener_qt_handler.prompt_fp_mismatch(
+                host, self.subject_cn, fingerprint_display_str)
+
+            if not result.proceed:
+                return False
+
+            if result.remember:
+                unsaved = True
+                credential.username = fingerprint_str
+
+        auth_jwt: Optional[str] = None
+        auth_password: Optional[str] = None
+        if not unsaved and not CREDENTIAL_STORE.encrypted_at_rest():
+            auth_jwt = credential.blob
+        else:
+            auth_password = credential.blob
 
-                self.client.get_message_response(msgid, msg,
-                                                 future=auth_response)
-            else:
-                self._loop.call_soon_threadsafe(auth_exception)
+        msgid, msg = message_protocol.create_authenticate_message(
+            jwt=auth_jwt, password=auth_password)
 
-        CREDENTIAL_STORE.get(
-            key=host,
-            context='Input authentication token',
-            context_long=f'Input authentication token for {host}\n\n'
-                         f'(Fingerprint: {fingerprint_display_str})',
-            on_credential_ready=on_credential_ready)
+        auth_response = asyncio.get_running_loop().create_future()
+        self.client.get_message_response(msgid, msg, future=auth_response)
 
         try:
             response = await auth_response
@@ -116,15 +111,18 @@ class SwitchboardClientProtocol(QuicConnectionProtocol):
 
             if auth_succeeded:
                 if CREDENTIAL_STORE.encrypted_at_rest():
-                    if closure_cred_unsaved and closure_credential:
-                        CREDENTIAL_STORE.set(host, closure_credential)
+                    if unsaved and credential:
+                        logging.info(f'Saving password for {host}')
+                        CREDENTIAL_STORE.set(host, credential)
                 elif 'jwt' in response:
                     credential = CredentialStore.Credential(
                         fingerprint_str, response['jwt'])
+                    logging.info(f'Saving authentication token for {host}')
                     CREDENTIAL_STORE.set(host, credential)
             else:
                 # Clear invalid saved credential
-                if not closure_cred_unsaved:
+                if not unsaved:
+                    logging.info(f'Forgetting incorrect password for {host}')
                     CREDENTIAL_STORE.set(host, None)
 
             return auth_succeeded
@@ -132,11 +130,153 @@ class SwitchboardClientProtocol(QuicConnectionProtocol):
             return False
 
 
+class FingerprintMismatchUserAction:
+    def __init__(self):
+        self.proceed: bool = False
+        self.remember: bool = False
+
+
 class ListenerQtHandler(QtCore.QObject):
     listener_connecting = QtCore.Signal(object)
     listener_connected = QtCore.Signal(object)
     listener_connection_failed = QtCore.Signal(object)
 
+    class Prompt(ABC):
+        @abstractmethod
+        def display(self):
+            pass
+
+    class PromptPassword(Prompt):
+        def __init__(
+            self,
+            host: str,
+            subject_cn: str,
+            fingerprint_desc: str,
+        ):
+            self.host = host
+            self.subject_cn = subject_cn
+            self.fingerprint_desc = fingerprint_desc
+            self.future: Optional[Future[str]] = None
+
+        def display(self):
+            blob, ok = QtWidgets.QInputDialog.getText(
+                None,
+                'Enter Password',
+                f'Enter password for {self.host} ({self.subject_cn})\n\n'
+                f'(Fingerprint: {self.fingerprint_desc})',
+                QtWidgets.QLineEdit.EchoMode.Password)
+
+            assert self.future
+            if blob and ok:
+                self.future.set_result(blob)
+            else:
+                self.future.set_result('')
+
+
+    class PromptFingerprintMismatch(Prompt):
+        def __init__(
+            self,
+            host: str,
+            subject_cn: str,
+            fingerprint_desc: str,
+        ):
+            self.host = host
+            self.subject_cn = subject_cn
+            self.fingerprint_desc = fingerprint_desc
+            self.future: Optional[Future[FingerprintMismatchUserAction]] = None
+
+        def display(self):
+            title = 'Warning: Potential security breach!'
+            text = (
+                f'{title}'
+                '\n\n'
+                'The Switchboard Listener key does not match the key that has '
+                'been cached. This means that either another user has changed '
+                'the key, or you are actually trying to connect to another '
+                'computer pretending to be the device.'
+                '\n\n'
+                'If the key change was not expected, please contact your '
+                'administrator.'
+                '\n\n'
+                'Details for new key:\n'
+                f'Host: {self.host} ({self.subject_cn})\n'
+                f'Fingerprint: {self.fingerprint_desc}'
+                '\n\n'
+                'Trust the new key and carry on connecting?\n'
+            )
+            check_text = 'Update cached key for this device'
+
+            dialog = QMessageBox()
+            dialog.setWindowTitle(title)
+            dialog.setText(text)
+
+            checkbox = QtWidgets.QCheckBox(check_text)
+            dialog.setCheckBox(checkbox)
+
+            dialog.setIconPixmap(QtGui.QIcon(
+                ':icons/images/alert-triangle-64.svg').pixmap(48, 48))
+
+            ok_btn = dialog.addButton(QMessageBox.StandardButton.Ok)
+            dialog.addButton(QMessageBox.StandardButton.Cancel)
+            dialog.setDefaultButton(QMessageBox.StandardButton.NoButton)
+
+            dialog.exec()
+
+            result = FingerprintMismatchUserAction()
+
+            if dialog.checkBox().isChecked():
+                result.remember = True
+
+            if dialog.clickedButton() == ok_btn:
+                result.proceed = True
+
+            assert self.future
+            self.future.set_result(result)
+
+    def __init__(self):
+        super().__init__()
+
+        self._queued_prompts: list[ListenerQtHandler.Prompt] = []
+
+    async def prompt_password(
+        self,
+        host: str,
+        subject_cn: str,
+        fingerprint_desc: str,
+    ) -> str:
+        prompt = ListenerQtHandler.PromptPassword(host, subject_cn,
+                                                  fingerprint_desc)
+        prompt.future = asyncio.get_running_loop().create_future()
+        self._queued_prompts.append(prompt)
+        self._show_queued_prompts_soon()
+        return await prompt.future
+
+    async def prompt_fp_mismatch(
+        self,
+        host: str,
+        subject_cn: str,
+        fingerprint_desc: str,
+    ) -> FingerprintMismatchUserAction:
+        prompt = ListenerQtHandler.PromptFingerprintMismatch(
+            host, subject_cn, fingerprint_desc)
+        prompt.future = asyncio.get_running_loop().create_future()
+        self._queued_prompts.append(prompt)
+        self._show_queued_prompts_soon()
+        return await prompt.future
+
+    @QtCore.Slot()
+    def _show_queued_prompts_soon(self):
+        # Ensure this code is run from the main thread
+        if threading.current_thread() is not threading.main_thread():
+            QtCore.QMetaObject.invokeMethod(
+                self, '_show_queued_prompts_soon', QtCore.Qt.QueuedConnection)
+            return
+
+        for prompt in self._queued_prompts:
+            prompt.display()
+
+        self._queued_prompts.clear()
+
 
 class ListenerClient:
     '''Connects to a server running SwitchboardListener.
@@ -243,7 +383,7 @@ class ListenerClient:
         return True
 
     def disconnect(self, unexpected=False, exception=None):
-        if self.protocol and self.loop:
+        if self.protocol and self.loop and (not self.loop.is_closed()):
             _, msg = message_protocol.create_disconnect_message()
             self.send_message(msg)
             self.loop.call_soon_threadsafe(self.protocol.close)
@@ -446,8 +586,8 @@ class ListenerClient:
         message_bytes: bytes,
         *,
         max_wait=IDLE_TIMEOUT_SEC,
-        future: Optional[asyncio.Future[dict[str, str]]] = None
-    ) -> asyncio.Future[dict[str, str]]:
+        future: Optional[Future[dict[str, str]]] = None
+    ) -> Future[dict[str, str]]:
         if self.loop and self.protocol and self.writer:
             if not future:
                 future = self.loop.create_future()
diff --git a/Engine/Plugins/VirtualProduction/Switchboard/Source/Switchboard/switchboard/resources.py b/Engine/Plugins/VirtualProduction/Switchboard/Source/Switchboard/switchboard/resources.py
--- a/Engine/Plugins/VirtualProduction/Switchboard/Source/Switchboard/switchboard/resources.py
+++ b/Engine/Plugins/VirtualProduction/Switchboard/Source/Switchboard/switchboard/resources.py
@@ -1810,6 +1810,262 @@ t\xff+O\x1b\x17\x91Z\x0e\x08\x22\x0fE\x93\xc4$\
 \x98\x0e*\xb6>aB\x15,\x11\xc0\xba\x1eh\xb3\x07\
 \x00\xe3$\x0b\x12\xdc!\x7f\xef\x00\x00\x00\x00IEN\
 D\xaeB`\x82\
+\x00\x00\x0f\xd8\
+<\
+svg width=\x2264\x22 h\
+eight=\x2264\x22 viewB\
+ox=\x220 0 64 64\x22 f\
+ill=\x22none\x22 xmlns\
+=\x22http://www.w3.\
+org/2000/svg\x22>\x0d\x0a\
+<g filter=\x22url(#\
+filter0_i_13504_\
+12386)\x22>\x0d\x0a<path \
+d=\x22M55 50.5002L3\
+5.3198 10.0002C3\
+4.9887 9.39435 3\
+4.5065 8.8845 33\
+.9199 8.52026C33\
+.3119 8.16919 32\
+.6221 7.98438 31\
+.9199 7.98438C31\
+.2178 7.98438 30\
+.528 8.16919 29.\
+9199 8.52026C29.\
+3334 8.8845 28.8\
+511 9.39435 28.5\
+2 10.0002L8.8799\
+ 50.4402C8.55246\
+ 51.0103 8.39953\
+ 51.664 8.43995 \
+52.3202C8.45184 \
+52.9681 8.63098 \
+53.602 8.95997 5\
+4.1603C9.28727 5\
+4.7139 9.75674 5\
+5.1696 10.3198 5\
+5.4802C10.8781 5\
+5.8092 11.5123 5\
+5.9884 12.1602 5\
+6.0002H51.52C52.\
+1679 55.9884 52.\
+8016 55.8092 53.\
+3599 55.4802C53.\
+923 55.1696 54.3\
+929 54.7139 54.7\
+202 54.1603C55.0\
+492 53.602 55.22\
+84 52.9681 55.24\
+02 52.3202C55.24\
+02 52.0002 55.28\
+03 51.0837 55 50\
+.5002Z\x22 stroke=\x22\
+#FAA501\x22 stroke-\
+width=\x223\x22/>\x0d\x0a</g\
+>\x0d\x0a<path fill-ru\
+le=\x22evenodd\x22 cli\
+p-rule=\x22evenodd\x22\
+ d=\x22M29.4199 7.6\
+5424C30.18 7.215\
+39 31.0422 6.984\
+38 31.9199 6.984\
+38C32.7976 6.984\
+38 33.6598 7.215\
+39 34.4199 7.654\
+24L34.4338 7.662\
+26L34.4474 7.670\
+72C35.1807 8.126\
+04 35.7835 8.763\
+38 36.1973 9.520\
+7L36.2089 9.5417\
+8L55 50.5002L34.\
+4315 10.4603C34.\
+1861 10.019 33.8\
+335 9.6469 33.40\
+6 9.3783C32.9534\
+ 9.12019 32.4412\
+ 8.98438 31.9199\
+ 8.98438C31.3986\
+ 8.98438 30.8865\
+ 9.12019 30.4338\
+ 9.3783C30.0064 \
+9.6469 29.6537 1\
+0.019 29.4083 10\
+.4603L9 50.5002L\
+27.631 9.54178L2\
+7.6425 9.5207C28\
+.0564 8.76338 28\
+.6591 8.12604 29\
+.3924 7.67072L29\
+.406 7.66226L29.\
+4199 7.65424Z\x22 f\
+ill=\x22#DDC085\x22/>\x0d\
+\x0a<g filter=\x22url(\
+#filter1_i_13504\
+_12386)\x22>\x0d\x0a<path\
+ d=\x22M32.002 49.0\
+268C31.0793 49.0\
+245 30.1861 48.7\
+024 29.4743 48.1\
+154C28.7625 47.5\
+284 28.276 46.71\
+28 28.0981 45.80\
+75C27.9201 44.90\
+22 28.0618 43.96\
+32 28.4985 43.15\
+05C28.9352 42.33\
+77 29.6399 41.70\
+15 30.4931 41.35\
+03C31.3462 40.99\
+9 32.295 40.9544\
+ 33.1773 41.224C\
+34.0597 41.4936 \
+34.821 42.0607 3\
+5.3321 42.8288C3\
+5.8433 43.5969 3\
+6.0722 44.5185 3\
+5.9801 45.4365C3\
+5.888 46.3545 35\
+.4805 47.2122 34\
+.827 47.8634C34.\
+071 48.6018 33.0\
+586 49.0187 32.0\
+02 49.0268ZM32.0\
+02 38.3909C31.12\
+05 38.3909 30.27\
+49 38.0407 29.65\
+16 37.4174C29.02\
+83 36.7941 28.67\
+83 35.9487 28.67\
+83 35.0672V25.09\
+62C28.6783 24.21\
+47 29.0283 23.36\
+93 29.6516 22.74\
+6C30.2749 22.122\
+7 31.1205 21.772\
+5 32.002 21.7725\
+C32.8835 21.7725\
+ 33.729 22.1227 \
+34.3523 22.746C3\
+4.9756 23.3693 3\
+5.3257 24.2147 3\
+5.3257 25.0962V3\
+5.0672C35.3257 3\
+5.9487 34.9756 3\
+6.7941 34.3523 3\
+7.4174C33.729 38\
+.0407 32.8835 38\
+.3909 32.002 38.\
+3909Z\x22 fill=\x22url\
+(#paint0_linear_\
+13504_12386)\x22/>\x0d\
+\x0a</g>\x0d\x0a<defs>\x0d\x0a<\
+filter id=\x22filte\
+r0_i_13504_12386\
+\x22 x=\x226.93359\x22 y=\
+\x223.48438\x22 width=\
+\x2249.8087\x22 height\
+=\x2254.0156\x22 filte\
+rUnits=\x22userSpac\
+eOnUse\x22 color-in\
+terpolation-filt\
+ers=\x22sRGB\x22>\x0d\x0a<fe\
+Flood flood-opac\
+ity=\x220\x22 result=\x22\
+BackgroundImageF\
+ix\x22/>\x0d\x0a<feBlend \
+mode=\x22normal\x22 in\
+=\x22SourceGraphic\x22\
+ in2=\x22Background\
+ImageFix\x22 result\
+=\x22shape\x22/>\x0d\x0a<feC\
+olorMatrix in=\x22S\
+ourceAlpha\x22 type\
+=\x22matrix\x22 values\
+=\x220 0 0 0 0 0 0 \
+0 0 0 0 0 0 0 0 \
+0 0 0 127 0\x22 res\
+ult=\x22hardAlpha\x22/\
+>\x0d\x0a<feOffset dy=\
+\x22-3\x22/>\x0d\x0a<feGauss\
+ianBlur stdDevia\
+tion=\x222\x22/>\x0d\x0a<feC\
+omposite in2=\x22ha\
+rdAlpha\x22 operato\
+r=\x22arithmetic\x22 k\
+2=\x22-1\x22 k3=\x221\x22/>\x0d\
+\x0a<feColorMatrix \
+type=\x22matrix\x22 va\
+lues=\x220 0 0 0 1 \
+0 0 0 0 1 0 0 0 \
+0 1 0 0 0 0.3 0\x22\
+/>\x0d\x0a<feBlend mod\
+e=\x22normal\x22 in2=\x22\
+shape\x22 result=\x22e\
+ffect1_innerShad\
+ow_13504_12386\x22/\
+>\x0d\x0a</filter>\x0d\x0a<f\
+ilter id=\x22filter\
+1_i_13504_12386\x22\
+ x=\x2228.0233\x22 y=\x22\
+19.2797\x22 width=\x22\
+7.97675\x22 height=\
+\x2229.7467\x22 filter\
+Units=\x22userSpace\
+OnUse\x22 color-int\
+erpolation-filte\
+rs=\x22sRGB\x22>\x0d\x0a<feF\
+lood flood-opaci\
+ty=\x220\x22 result=\x22B\
+ackgroundImageFi\
+x\x22/>\x0d\x0a<feBlend m\
+ode=\x22normal\x22 in=\
+\x22SourceGraphic\x22 \
+in2=\x22BackgroundI\
+mageFix\x22 result=\
+\x22shape\x22/>\x0d\x0a<feCo\
+lorMatrix in=\x22So\
+urceAlpha\x22 type=\
+\x22matrix\x22 values=\
+\x220 0 0 0 0 0 0 0\
+ 0 0 0 0 0 0 0 0\
+ 0 0 127 0\x22 resu\
+lt=\x22hardAlpha\x22/>\
+\x0d\x0a<feOffset dy=\x22\
+-2.49277\x22/>\x0d\x0a<fe\
+GaussianBlur std\
+Deviation=\x221.661\
+84\x22/>\x0d\x0a<feCompos\
+ite in2=\x22hardAlp\
+ha\x22 operator=\x22ar\
+ithmetic\x22 k2=\x22-1\
+\x22 k3=\x221\x22/>\x0d\x0a<feC\
+olorMatrix type=\
+\x22matrix\x22 values=\
+\x220 0 0 0 1 0 0 0\
+ 0 1 0 0 0 0 1 0\
+ 0 0 0.3 0\x22/>\x0d\x0a<\
+feBlend mode=\x22no\
+rmal\x22 in2=\x22shape\
+\x22 result=\x22effect\
+1_innerShadow_13\
+504_12386\x22/>\x0d\x0a</\
+filter>\x0d\x0a<linear\
+Gradient id=\x22pai\
+nt0_linear_13504\
+_12386\x22 x1=\x2232.0\
+306\x22 y1=\x2222.0059\
+\x22 x2=\x2231.583\x22 y2\
+=\x2249.0193\x22 gradi\
+entUnits=\x22userSp\
+aceOnUse\x22>\x0d\x0a<sto\
+p stop-color=\x22#F\
+FA800\x22/>\x0d\x0a<stop \
+offset=\x221\x22 stop-\
+color=\x22#E39600\x22/\
+>\x0d\x0a</linearGradi\
+ent>\x0d\x0a</defs>\x0d\x0a<\
+/svg>\x0d\x0a\
 \x00\x00\x01w\
 \x89\
 PNG\x0d\x0a\x1a\x0a\x00\x00\x00\x0dIHDR\x00\
@@ -9089,6 +9345,11 @@ qt_resource_name = b"\
 \x00t\
 \x00r\x00e\x00e\x00_\x00a\x00r\x00r\x00o\x00w\x00_\x00c\x00o\x00l\x00l\x00a\x00p\
 \x00s\x00e\x00d\x00.\x00p\x00n\x00g\
+\x00\x15\
+\x0cOq\x07\
+\x00a\
+\x00l\x00e\x00r\x00t\x00-\x00t\x00r\x00i\x00a\x00n\x00g\x00l\x00e\x00-\x006\x004\
+\x00.\x00s\x00v\x00g\
 \x00\x14\
 \x08b\xb6\x87\
 \x00i\
@@ -9456,7 +9717,7 @@ qt_resource_name = b"\
 qt_resource_struct = b"\
 \x00\x00\x00\x00\x00\x02\x00\x00\x00\x03\x00\x00\x00\x01\
 \x00\x00\x00\x00\x00\x00\x00\x00\
-\x00\x00\x00\x22\x00\x02\x00\x00\x00\x01\x00\x00\x00n\
+\x00\x00\x00\x22\x00\x02\x00\x00\x00\x01\x00\x00\x00o\
 \x00\x00\x00\x00\x00\x00\x00\x00\
 \x00\x00\x00\x00\x00\x02\x00\x00\x00\x01\x00\x00\x00\x06\
 \x00\x00\x00\x00\x00\x00\x00\x00\
@@ -9465,217 +9726,219 @@ qt_resource_struct = b"\
 \x00\x00\x00\x10\x00\x02\x00\x00\x00\x01\x00\x00\x00\x05\
 \x00\x00\x00\x00\x00\x00\x00\x00\
 \x00\x00\x00R\x00\x00\x00\x00\x00\x01\x00\x00\x0a\x0e\
-\x00\x00\x01\x8e5\x01\xe3\xad\
-\x00\x00\x00\x10\x00\x02\x00\x00\x00g\x00\x00\x00\x07\
+\x00\x00\x01\x87\xbe\xcb4\xd1\
+\x00\x00\x00\x10\x00\x02\x00\x00\x00h\x00\x00\x00\x07\
 \x00\x00\x00\x00\x00\x00\x00\x00\
-\x00\x00\x0a\x08\x00\x00\x00\x00\x00\x01\x00\x00\xea@\
-\x00\x00\x01\x8e5\x01\xe4*\
-\x00\x00\x05\xde\x00\x00\x00\x00\x00\x01\x00\x00x%\
-\x00\x00\x01\x8e5\x01\xe4\xb5\
-\x00\x00\x13,\x00\x00\x00\x00\x00\x01\x00\x02\x1ey\
-\x00\x00\x01\x8e5\x01\xe4%\
+\x00\x00\x0a8\x00\x00\x00\x00\x00\x01\x00\x00\xfa\x1c\
+\x00\x00\x01\x87\xbe\xcb5\x1a\
+\x00\x00\x06\x0e\x00\x00\x00\x00\x00\x01\x00\x00\x88\x01\
+\x00\x00\x01\x87\xbe\xcb5\xa4\
+\x00\x00\x13\x5c\x00\x00\x00\x00\x00\x01\x00\x02.U\
+\x00\x00\x01\x87\xbe\xcb5\x14\
 \x00\x00\x04\x88\x00\x00\x00\x00\x00\x01\x00\x00[\x95\
-\x00\x00\x01\x8e5\x01\xe3\xb1\
-\x00\x00\x09\x18\x00\x00\x00\x00\x00\x01\x00\x00\xdb3\
-\x00\x00\x01\x8e5\x01\xe4\x96\
+\x00\x00\x01\x87\xbe\xcb4\xd3\
+\x00\x00\x09H\x00\x00\x00\x00\x00\x01\x00\x00\xeb\x0f\
+\x00\x00\x01\x87\xbe\xcb5\x7f\
 \x00\x00\x03H\x00\x00\x00\x00\x00\x01\x00\x00<V\
-\x00\x00\x01\x8e5\x01\xe4<\
-\x00\x00\x08\xe4\x00\x00\x00\x00\x00\x01\x00\x00\xbe\x90\
-\x00\x00\x01\x8e5\x01\xe3\xbd\
-\x00\x00\x12L\x00\x00\x00\x00\x00\x01\x00\x01\xf5\x04\
-\x00\x00\x01\x8e5\x01\xe3\xc7\
-\x00\x00\x13\x08\x00\x00\x00\x00\x00\x01\x00\x02\x0c\xf1\
-\x00\x00\x01\x8e5\x01\xe4\x9e\
-\x00\x00\x08\x16\x00\x00\x00\x00\x00\x01\x00\x00\xb8\xcb\
-\x00\x00\x01\x8e5\x01\xe4z\
+\x00\x00\x01\x87\xbe\xcb5,\
+\x00\x00\x09\x14\x00\x00\x00\x00\x00\x01\x00\x00\xcel\
+\x00\x00\x01\x87\xbe\xcb4\xd9\
+\x00\x00\x12|\x00\x00\x00\x00\x00\x01\x00\x02\x04\xe0\
+\x00\x00\x01\x87\xbe\xcb4\xdd\
+\x00\x00\x138\x00\x00\x00\x00\x00\x01\x00\x02\x1c\xcd\
+\x00\x00\x01\x87\xbe\xcb5\x86\
+\x00\x00\x08F\x00\x00\x00\x00\x00\x01\x00\x00\xc8\xa7\
+\x00\x00\x01\x87\xbe\xcb5c\
 \x00\x00\x00~\x00\x00\x00\x00\x00\x01\x00\x00\x10b\
-\x00\x00\x01\x8e5\x01\xe3\x9d\
-\x00\x00\x08l\x00\x00\x00\x00\x00\x01\x00\x00\xba\xa2\
-\x00\x00\x01\x8e5\x01\xe3\xd1\
-\x00\x00\x06V\x00\x00\x00\x00\x00\x01\x00\x00{\x0e\
-\x00\x00\x01\x8e5\x01\xe4d\
+\x00\x00\x01\x87\xbe\xcb4\xca\
+\x00\x00\x08\x9c\x00\x00\x00\x00\x00\x01\x00\x00\xca~\
+\x00\x00\x01\x87\xbe\xcb4\xe1\
+\x00\x00\x06\x86\x00\x00\x00\x00\x00\x01\x00\x00\x8a\xea\
+\x00\x00\x01\x87\xbe\xcb5K\
 \x00\x00\x02\x82\x00\x00\x00\x00\x00\x01\x00\x00%\xaa\
-\x00\x00\x01\x8e5\x01\xe4\xc9\
-\x00\x00\x0f\x9a\x00\x00\x00\x00\x00\x01\x00\x01\x98\xe9\
-\x00\x00\x01\x8e5\x01\xe4\x15\
-\x00\x00\x0f\x04\x00\x00\x00\x00\x00\x01\x00\x01\x8d4\
-\x00\x00\x01\x8e5\x01\xe4\x05\
+\x00\x00\x01\x87\xbe\xcb5\xbd\
+\x00\x00\x0f\xca\x00\x00\x00\x00\x00\x01\x00\x01\xa8\xc5\
+\x00\x00\x01\x88\xf8\xc2W\x00\
+\x00\x00\x0f4\x00\x00\x00\x00\x00\x01\x00\x01\x9d\x10\
+\x00\x00\x01\x88\xf8\xc2V\xfe\
 \x00\x00\x02\xdc\x00\x00\x00\x00\x00\x01\x00\x003\x0e\
-\x00\x00\x01\x8e5\x01\xe4\x05\
-\x00\x00\x10\xec\x00\x00\x00\x00\x00\x01\x00\x01\xcd\xf0\
-\x00\x00\x01\x8e5\x01\xe4\x0c\
-\x00\x00\x06\x9a\x00\x00\x00\x00\x00\x01\x00\x00|\x13\
-\x00\x00\x01\x8e5\x01\xe4\x0b\
+\x00\x00\x01\x88\xf8\xc2V\xfd\
+\x00\x00\x11\x1c\x00\x00\x00\x00\x00\x01\x00\x01\xdd\xcc\
+\x00\x00\x01\x88\xf8\xc2V\xff\
+\x00\x00\x06\xca\x00\x00\x00\x00\x00\x01\x00\x00\x8b\xef\
+\x00\x00\x01\x88\xf8\xc2V\xff\
 \x00\x00\x03\x12\x00\x00\x00\x00\x00\x01\x00\x007\xc4\
-\x00\x00\x01\x8e5\x01\xe3\xf8\
-\x00\x00\x0c.\x00\x00\x00\x00\x00\x01\x00\x01S'\
-\x00\x00\x01\x8e5\x01\xe3\xf5\
+\x00\x00\x01\x88\xf8\xc2V\xfa\
+\x00\x00\x0c^\x00\x00\x00\x00\x00\x01\x00\x01c\x03\
+\x00\x00\x01\x88\xf8\xc2V\xf8\
 \x00\x00\x04\xb8\x00\x00\x00\x00\x00\x01\x00\x00\x5c+\
-\x00\x00\x01\x8e5\x01\xe3\xfe\
-\x00\x00\x0e\x9a\x00\x00\x00\x00\x00\x01\x00\x01}\x82\
-\x00\x00\x01\x8e5\x01\xe3\xfe\
+\x00\x00\x01\x88\xf8\xc2V\xfb\
+\x00\x00\x0e\xca\x00\x00\x00\x00\x00\x01\x00\x01\x8d^\
+\x00\x00\x01\x88\xf8\xc2V\xfa\
 \x00\x00\x00\xf2\x00\x00\x00\x00\x00\x01\x00\x00\x1cc\
-\x00\x00\x01\x8e5\x01\xe3\xeb\
-\x00\x00\x11\xc6\x00\x00\x00\x00\x00\x01\x00\x01\xeb\x17\
-\x00\x00\x01\x8e5\x01\xe4\xa7\
+\x00\x00\x01\x87\xbe\xcb4\xe9\
+\x00\x00\x11\xf6\x00\x00\x00\x00\x00\x01\x00\x01\xfa\xf3\
+\x00\x00\x01\x87\xbe\xcb5\x92\
 \x00\x00\x02b\x00\x00\x00\x00\x00\x01\x00\x00%9\
-\x00\x00\x01\x8e5\x01\xe4[\
-\x00\x00\x06\xf0\x00\x00\x00\x00\x00\x01\x00\x00\x81\x8a\
-\x00\x00\x01\x8e5\x01\xe3\xf4\
-\x00\x00\x10p\x00\x00\x00\x00\x00\x01\x00\x01\xcb~\
-\x00\x00\x01\x8e5\x01\xe3\xeb\
-\x00\x00\x07\x90\x00\x00\x00\x00\x00\x01\x00\x00\xaa\xdd\
-\x00\x00\x01\x8e5\x01\xe44\
+\x00\x00\x01\x87\xbe\xcb5G\
+\x00\x00\x07 \x00\x00\x00\x00\x00\x01\x00\x00\x91f\
+\x00\x00\x01\x87\xbe\xcb4\xed\
+\x00\x00\x10\xa0\x00\x00\x00\x00\x00\x01\x00\x01\xdbZ\
+\x00\x00\x01\x87\xbe\xcb4\xeb\
+\x00\x00\x07\xc0\x00\x00\x00\x00\x00\x01\x00\x00\xba\xb9\
+\x00\x00\x01\x87\xbe\xcb5%\
 \x00\x00\x04N\x00\x00\x00\x00\x00\x01\x00\x00X8\
-\x00\x00\x01\x8e5\x01\xe4R\
-\x00\x00\x08\xb2\x00\x00\x00\x00\x00\x01\x00\x00\xbcf\
-\x00\x00\x01\x8e5\x01\xe4<\
-\x00\x00\x10\xc0\x00\x00\x00\x00\x00\x01\x00\x01\xcdY\
-\x00\x00\x01\x8e5\x01\xe4d\
+\x00\x00\x01\x87\xbe\xcb5=\
+\x00\x00\x08\xe2\x00\x00\x00\x00\x00\x01\x00\x00\xccB\
+\x00\x00\x01\x87\xbe\xcb5-\
+\x00\x00\x10\xf0\x00\x00\x00\x00\x00\x01\x00\x01\xdd5\
+\x00\x00\x01\x87\xbe\xcb5I\
 \x00\x00\x03z\x00\x00\x00\x00\x00\x01\x00\x00>\xda\
-\x00\x00\x01\x8e5\x01\xe3\xa5\
+\x00\x00\x01\x87\xbe\xcb4\xcf\
 \x00\x00\x01\x98\x00\x00\x00\x00\x00\x01\x00\x00!<\
-\x00\x00\x01\x8e5\x01\xe4K\
-\x00\x00\x0e\xd0\x00\x00\x00\x00\x00\x01\x00\x01\x82\x22\
-\x00\x00\x01\x8e5\x01\xe4\xbe\
-\x00\x00\x09\xb0\x00\x00\x00\x00\x00\x01\x00\x00\xe7\xe7\
-\x00\x00\x01\x8e5\x01\xe4\x96\
-\x00\x00\x07\xb6\x00\x00\x00\x00\x00\x01\x00\x00\xac\x1e\
-\x00\x00\x01\x8e5\x01\xe3\xc7\
-\x00\x00\x09\xde\x00\x00\x00\x00\x00\x01\x00\x00\xe8\xfd\
-\x00\x00\x01\x8e5\x01\xe3\xd7\
-\x00\x00\x11\xec\x00\x00\x00\x00\x00\x01\x00\x01\xeb\xc8\
-\x00\x00\x01\x8e5\x01\xe4L\
-\x00\x00\x12&\x00\x00\x00\x00\x00\x01\x00\x01\xed\x9f\
-\x00\x00\x01\x8e5\x01\xe4\x8d\
-\x00\x00\x12\xc4\x00\x00\x00\x00\x00\x01\x00\x02\x01\xd8\
-\x00\x00\x01\x8e5\x01\xe4\xc4\
-\x00\x00\x0b\x8c\x00\x00\x00\x00\x00\x01\x00\x01E;\
-\x00\x00\x01\x8e5\x01\xe3\xe0\
-\x00\x00\x10\x14\x00\x00\x00\x00\x00\x01\x00\x01\x9f\x5c\
-\x00\x00\x01\x8e5\x01\xe4\xa7\
-\x00\x00\x0b\xc2\x00\x00\x00\x00\x00\x01\x00\x01Fj\
-\x00\x00\x01\x8e5\x01\xe4$\
-\x00\x00\x11\x22\x00\x00\x00\x00\x00\x01\x00\x01\xd2\xb4\
-\x00\x00\x01\x8e5\x01\xe3\xb6\
-\x00\x00\x05\xac\x00\x00\x00\x00\x00\x01\x00\x00o\x03\
-\x00\x00\x01\x8e5\x01\xe4\x8e\
-\x00\x00\x0b\xee\x00\x00\x00\x00\x00\x01\x00\x01Gq\
-\x00\x00\x01\x8e5\x01\xe3\x9a\
-\x00\x00\x0a6\x00\x00\x00\x00\x00\x01\x00\x00\xec\x9d\
-\x00\x00\x01\x8e5\x01\xe4E\
-\x00\x00\x05z\x00\x00\x00\x00\x00\x01\x00\x00nm\
-\x00\x00\x01\x8e5\x01\xe4s\
-\x00\x00\x11\x5c\x00\x00\x00\x00\x00\x01\x00\x01\xe5\xd9\
-\x00\x00\x01\x8e5\x01\xe4c\
-\x00\x00\x05L\x00\x00\x00\x00\x00\x01\x00\x00l\xf2\
-\x00\x00\x01\x8e5\x01\xe3\xd0\
-\x00\x00\x07h\x00\x00\x00\x00\x00\x01\x00\x00\xa7\x8a\
-\x00\x00\x01\x8e5\x01\xe4R\
-\x00\x00\x0d\xd6\x00\x00\x00\x00\x00\x01\x00\x01pY\
-\x00\x00\x01\x8e5\x01\xe4\xbc\
-\x00\x00\x0af\x00\x00\x00\x00\x00\x01\x00\x00\xed-\
-\x00\x00\x01\x8e5\x01\xe3\xb1\
-\x00\x00\x0a\x90\x00\x00\x00\x00\x00\x01\x00\x00\xed\xe8\
-\x00\x00\x01\x8e5\x01\xe4\x15\
+\x00\x00\x01\x87\xbe\xcb55\
+\x00\x00\x0f\x00\x00\x00\x00\x00\x00\x01\x00\x01\x91\xfe\
+\x00\x00\x01\x87\xbe\xcb5\xab\
+\x00\x00\x09\xe0\x00\x00\x00\x00\x00\x01\x00\x00\xf7\xc3\
+\x00\x00\x01\x87\xbe\xcb5~\
+\x00\x00\x07\xe6\x00\x00\x00\x00\x00\x01\x00\x00\xbb\xfa\
+\x00\x00\x01\x87\xbe\xcb4\xdc\
+\x00\x00\x0a\x0e\x00\x00\x00\x00\x00\x01\x00\x00\xf8\xd9\
+\x00\x00\x01\x87\xbe\xcb4\xe5\
+\x00\x00\x12\x1c\x00\x00\x00\x00\x00\x01\x00\x01\xfb\xa4\
+\x00\x00\x01\x87\xbe\xcb56\
+\x00\x00\x12V\x00\x00\x00\x00\x00\x01\x00\x01\xfd{\
+\x00\x00\x01\x87\xbe\xcb5u\
+\x00\x00\x12\xf4\x00\x00\x00\x00\x00\x01\x00\x02\x11\xb4\
+\x00\x00\x01\x87\xbe\xcb5\xb4\
+\x00\x00\x0b\xbc\x00\x00\x00\x00\x00\x01\x00\x01U\x17\
+\x00\x00\x01\x87\xbe\xcb4\xe7\
+\x00\x00\x10D\x00\x00\x00\x00\x00\x01\x00\x01\xaf8\
+\x00\x00\x01\x87\xbe\xcb5\x90\
+\x00\x00\x0b\xf2\x00\x00\x00\x00\x00\x01\x00\x01VF\
+\x00\x00\x01\x87\xbe\xcb5\x11\
+\x00\x00\x11R\x00\x00\x00\x00\x00\x01\x00\x01\xe2\x90\
+\x00\x00\x01\x87\xbe\xcb4\xd7\
+\x00\x00\x05\xdc\x00\x00\x00\x00\x00\x01\x00\x00~\xdf\
+\x00\x00\x01\x87\xbe\xcb5y\
+\x00\x00\x0c\x1e\x00\x00\x00\x00\x00\x01\x00\x01WM\
+\x00\x00\x01\x87\xbe\xcb4\xc8\
+\x00\x00\x0af\x00\x00\x00\x00\x00\x01\x00\x00\xfcy\
+\x00\x00\x01\x87\xbe\xcb52\
+\x00\x00\x05\xaa\x00\x00\x00\x00\x00\x01\x00\x00~I\
+\x00\x00\x01\x87\xbe\xcb5\x5c\
+\x00\x00\x11\x8c\x00\x00\x00\x00\x00\x01\x00\x01\xf5\xb5\
+\x00\x00\x01\x87\xbe\xcb5H\
+\x00\x00\x05|\x00\x00\x00\x00\x00\x01\x00\x00|\xce\
+\x00\x00\x01\x87\xbe\xcb4\xe0\
+\x00\x00\x07\x98\x00\x00\x00\x00\x00\x01\x00\x00\xb7f\
+\x00\x00\x01\x87\xbe\xcb5<\
+\x00\x00\x0e\x06\x00\x00\x00\x00\x00\x01\x00\x01\x805\
+\x00\x00\x01\x87\xbe\xcb5\xa7\
+\x00\x00\x0a\x96\x00\x00\x00\x00\x00\x01\x00\x00\xfd\x09\
+\x00\x00\x01\x87\xbe\xcb4\xd4\
+\x00\x00\x0a\xc0\x00\x00\x00\x00\x00\x01\x00\x00\xfd\xc4\
+\x00\x00\x01\x87\xbe\xcb5\x04\
 \x00\x00\x03\xc4\x00\x00\x00\x00\x00\x01\x00\x00K\x12\
-\x00\x00\x01\x8e5\x01\xe4\x1d\
-\x00\x00\x07\xe4\x00\x00\x00\x00\x00\x01\x00\x00\xb7\xb6\
-\x00\x00\x01\x8e5\x01\xe3\xeb\
-\x00\x00\x10\x90\x00\x00\x00\x00\x00\x01\x00\x01\xcc\x92\
-\x00\x00\x01\x8e5\x01\xe44\
-\x00\x00\x09\x84\x00\x00\x00\x00\x00\x01\x00\x00\xe6\xe9\
-\x00\x00\x01\x8e5\x01\xe4\x1d\
-\x00\x00\x0e^\x00\x00\x00\x00\x00\x01\x00\x01|@\
-\x00\x00\x01\x8e5\x01\xe3\xd7\
-\x00\x00\x12\x9a\x00\x00\x00\x00\x00\x01\x00\x01\xf6\xae\
-\x00\x00\x01\x8e5\x01\xe4\x8d\
-\x00\x00\x0d.\x00\x00\x00\x00\x00\x01\x00\x01e\x83\
-\x00\x00\x01\x8e5\x01\xe4g\
-\x00\x00\x0bP\x00\x00\x00\x00\x00\x01\x00\x00\xf9\x99\
-\x00\x00\x01\x8e5\x01\xe4\x87\
-\x00\x00\x11\x8e\x00\x00\x00\x00\x00\x01\x00\x01\xe6J\
-\x00\x00\x01\x8e5\x01\xe3\xf7\
+\x00\x00\x01\x87\xbe\xcb5\x0b\
+\x00\x00\x08\x14\x00\x00\x00\x00\x00\x01\x00\x00\xc7\x92\
+\x00\x00\x01\x87\xbe\xcb4\xec\
+\x00\x00\x10\xc0\x00\x00\x00\x00\x00\x01\x00\x01\xdcn\
+\x00\x00\x01\x87\xbe\xcb5$\
+\x00\x00\x09\xb4\x00\x00\x00\x00\x00\x01\x00\x00\xf6\xc5\
+\x00\x00\x01\x87\xbe\xcb5\x09\
+\x00\x00\x0e\x8e\x00\x00\x00\x00\x00\x01\x00\x01\x8c\x1c\
+\x00\x00\x01\x87\xbe\xcb4\xe6\
+\x00\x00\x12\xca\x00\x00\x00\x00\x00\x01\x00\x02\x06\x8a\
+\x00\x00\x01\x87\xbe\xcb5t\
+\x00\x00\x0d^\x00\x00\x00\x00\x00\x01\x00\x01u_\
+\x00\x00\x01\x87\xbe\xcb5P\
+\x00\x00\x0b\x80\x00\x00\x00\x00\x00\x01\x00\x01\x09u\
+\x00\x00\x01\x87\xbe\xcb5p\
+\x00\x00\x11\xbe\x00\x00\x00\x00\x00\x01\x00\x01\xf6&\
+\x00\x00\x01\x88\xf8\xc2V\xf9\
 \x00\x00\x01\xc6\x00\x00\x00\x00\x00\x01\x00\x00\x22W\
-\x00\x00\x01\x8e5\x01\xe4\x83\
-\x00\x00\x0f:\x00\x00\x00\x00\x00\x01\x00\x01\x91\xf7\
-\x00\x00\x01\x8e5\x01\xe4}\
-\x00\x00\x0bh\x00\x00\x00\x00\x00\x01\x00\x00\xfdr\
-\x00\x00\x01\x8e5\x01\xe4\xb4\
+\x00\x00\x01\x87\xbe\xcb5l\
+\x00\x00\x0fj\x00\x00\x00\x00\x00\x01\x00\x01\xa1\xd3\
+\x00\x00\x01\x87\xbe\xcb5f\
+\x00\x00\x0b\x98\x00\x00\x00\x00\x00\x01\x00\x01\x0dN\
+\x00\x00\x01\x87\xbe\xcb5\xa2\
 \x00\x00\x01,\x00\x00\x00\x00\x00\x01\x00\x00\x1dZ\
-\x00\x00\x01\x8e5\x01\xe4+\
-\x00\x00\x08\x92\x00\x00\x00\x00\x00\x01\x00\x00\xbb\xaf\
-\x00\x00\x01\x8e5\x01\xe4\x15\
-\x00\x00\x06v\x00\x00\x00\x00\x00\x01\x00\x00{\xa4\
-\x00\x00\x01\x8e5\x01\xe4D\
-\x00\x00\x0f\xf6\x00\x00\x00\x00\x00\x01\x00\x01\x9e\xc5\
-\x00\x00\x01\x8e5\x01\xe3\xad\
+\x00\x00\x01\x87\xbe\xcb5\x1c\
+\x00\x00\x08\xc2\x00\x00\x00\x00\x00\x01\x00\x00\xcb\x8b\
+\x00\x00\x01\x87\xbe\xcb5\x03\
+\x00\x00\x06\xa6\x00\x00\x00\x00\x00\x01\x00\x00\x8b\x80\
+\x00\x00\x01\x87\xbe\xcb5/\
+\x00\x00\x10&\x00\x00\x00\x00\x00\x01\x00\x01\xae\xa1\
+\x00\x00\x01\x87\xbe\xcb4\xd2\
 \x00\x00\x03\xf6\x00\x00\x00\x00\x00\x01\x00\x00K\xbd\
-\x00\x00\x01\x8e5\x01\xe3\x99\
-\x00\x00\x12n\x00\x00\x00\x00\x00\x01\x00\x01\xf5\xf6\
-\x00\x00\x01\x8e5\x01\xe4s\
+\x00\x00\x01\x87\xbe\xcb4\xc6\
+\x00\x00\x12\x9e\x00\x00\x00\x00\x00\x01\x00\x02\x05\xd2\
+\x00\x00\x01\x87\xbe\xcb5]\
 \x00\x00\x05\x16\x00\x00\x00\x00\x00\x01\x00\x00a\xd8\
-\x00\x00\x01\x8e5\x01\xe4\xb8\
-\x00\x00\x0c\xfc\x00\x00\x00\x00\x00\x01\x00\x01d\x91\
-\x00\x00\x01\x8e5\x01\xe3\xd6\
-\x00\x00\x0d\x9e\x00\x00\x00\x00\x00\x01\x00\x01o\xa8\
-\x00\x00\x01\x8e5\x01\xe4\xa7\
-\x00\x00\x0c\x86\x00\x00\x00\x00\x00\x01\x00\x01X\xfa\
-\x00\x00\x01\x8e5\x01\xe4M\
-\x00\x00\x10@\x00\x00\x00\x00\x00\x01\x00\x01\xbf\xc9\
-\x00\x00\x01\x8e5\x01\xe3\x99\
+\x00\x00\x01\x87\xbe\xcb5\xa5\
+\x00\x00\x0d,\x00\x00\x00\x00\x00\x01\x00\x01tm\
+\x00\x00\x01\x87\xbe\xcb4\xe3\
+\x00\x00\x0d\xce\x00\x00\x00\x00\x00\x01\x00\x01\x7f\x84\
+\x00\x00\x01\x87\xbe\xcb5\x93\
+\x00\x00\x0c\xb6\x00\x00\x00\x00\x00\x01\x00\x01h\xd6\
+\x00\x00\x01\x87\xbe\xcb57\
+\x00\x00\x10p\x00\x00\x00\x00\x00\x01\x00\x01\xcf\xa5\
+\x00\x00\x01\x87\xbe\xcb4\xc7\
 \x00\x00\x04\x16\x00\x00\x00\x00\x00\x01\x00\x00V\xf8\
-\x00\x00\x01\x8e5\x01\xe4;\
-\x00\x00\x0e\x1c\x00\x00\x00\x00\x00\x01\x00\x01{n\
-\x00\x00\x01\x8e5\x01\xe4\x83\
+\x00\x00\x01\x87\xbe\xcb5'\
+\x00\x00\x0eL\x00\x00\x00\x00\x00\x01\x00\x01\x8bJ\
+\x00\x00\x01\x87\xbe\xcb5m\
 \x00\x00\x03\x9c\x00\x00\x00\x00\x00\x01\x00\x00?\x5c\
-\x00\x00\x01\x8e5\x01\xe3\xaa\
+\x00\x00\x01\x87\xbe\xcb4\xd0\
 \x00\x00\x04\xee\x00\x00\x00\x00\x00\x01\x00\x00`\xde\
-\x00\x00\x01\x8e5\x01\xe4\xaf\
-\x00\x00\x0a\xf6\x00\x00\x00\x00\x00\x01\x00\x00\xf3i\
-\x00\x00\x01\x8e5\x01\xe4[\
+\x00\x00\x01\x87\xbe\xcb5\x9c\
+\x00\x00\x0b&\x00\x00\x00\x00\x00\x01\x00\x01\x03E\
+\x00\x00\x01\x87\xbe\xcb5F\
 \x00\x00\x01\xf6\x00\x00\x00\x00\x00\x01\x00\x00#+\
-\x00\x00\x01\x8e5\x01\xe3\xc8\
+\x00\x00\x01\x87\xbe\xcb4\xde\
 \x00\x00\x02\xa6\x00\x00\x00\x00\x00\x01\x00\x002\x9f\
-\x00\x00\x01\x8e5\x01\xe4E\
-\x00\x00\x0f\xd0\x00\x00\x00\x00\x00\x01\x00\x01\x9d\xc8\
-\x00\x00\x01\x8e5\x01\xe4\xae\
-\x00\x00\x0a\xc2\x00\x00\x00\x00\x00\x01\x00\x00\xee\x9e\
-\x00\x00\x01\x8e5\x01\xe3\xbf\
-\x00\x00\x05\xfe\x00\x00\x00\x00\x00\x01\x00\x00y}\
-\x00\x00\x01\x8e5\x01\xe4+\
-\x00\x00\x0fx\x00\x00\x00\x00\x00\x01\x00\x01\x93H\
-\x00\x00\x01\x8e5\x01\xe3\xbd\
-\x00\x00\x09>\x00\x00\x00\x00\x00\x01\x00\x00\xde\x93\
-\x00\x00\x01\x8e5\x01\xe4j\
-\x00\x00\x064\x00\x00\x00\x00\x00\x01\x00\x00z\x11\
-\x00\x00\x01\x8e5\x01\xe4\xb1\
+\x00\x00\x01\x87\xbe\xcb50\
+\x00\x00\x10\x00\x00\x00\x00\x00\x00\x01\x00\x01\xad\xa4\
+\x00\x00\x01\x87\xbe\xcb5\x9b\
+\x00\x00\x05L\x00\x00\x00\x00\x00\x01\x00\x00l\xf2\
+\x00\x00\x01\x87\xbe\xc8\xce\xe2\
+\x00\x00\x0a\xf2\x00\x00\x00\x00\x00\x01\x00\x00\xfez\
+\x00\x00\x01\x87\xbe\xcb4\xdb\
+\x00\x00\x06.\x00\x00\x00\x00\x00\x01\x00\x00\x89Y\
+\x00\x00\x01\x87\xbe\xcb5\x1d\
+\x00\x00\x0f\xa8\x00\x00\x00\x00\x00\x01\x00\x01\xa3$\
+\x00\x00\x01\x87\xbe\xcb4\xda\
+\x00\x00\x09n\x00\x00\x00\x00\x00\x01\x00\x00\xeeo\
+\x00\x00\x01\x87\xbe\xcb5S\
+\x00\x00\x06d\x00\x00\x00\x00\x00\x01\x00\x00\x89\xed\
+\x00\x00\x01\x87\xbe\xcb5\x9e\
 \x00\x00\x02*\x00\x00\x00\x00\x00\x01\x00\x00$\x1e\
-\x00\x00\x01\x8e5\x01\xe3\xd1\
-\x00\x00\x06\xd0\x00\x00\x00\x00\x00\x01\x00\x00\x80\xdd\
-\x00\x00\x01\x8e5\x01\xe4\x1d\
-\x00\x00\x08H\x00\x00\x00\x00\x00\x01\x00\x00\xb9\x9c\
-\x00\x00\x01\x8e5\x01\xe4\xae\
-\x00\x00\x0cd\x00\x00\x00\x00\x00\x01\x00\x01W\x92\
-\x00\x00\x01\x8e5\x01\xe4%\
-\x00\x00\x0c\xbc\x00\x00\x00\x00\x00\x01\x00\x01Z\x16\
-\x00\x00\x01\x8e5\x01\xe4k\
-\x00\x00\x0b*\x00\x00\x00\x00\x00\x01\x00\x00\xf7\x9f\
-\x00\x00\x01\x8e5\x01\xe3\xa5\
+\x00\x00\x01\x87\xbe\xcb4\xe2\
+\x00\x00\x07\x00\x00\x00\x00\x00\x00\x01\x00\x00\x90\xb9\
+\x00\x00\x01\x87\xbe\xcb5\x0a\
+\x00\x00\x08x\x00\x00\x00\x00\x00\x01\x00\x00\xc9x\
+\x00\x00\x01\x87\xbe\xcb5\x9a\
+\x00\x00\x0c\x94\x00\x00\x00\x00\x00\x01\x00\x01gn\
+\x00\x00\x01\x87\xbe\xcb5\x12\
+\x00\x00\x0c\xec\x00\x00\x00\x00\x00\x01\x00\x01i\xf2\
+\x00\x00\x01\x87\xbe\xcb5T\
+\x00\x00\x0bZ\x00\x00\x00\x00\x00\x01\x00\x01\x07{\
+\x00\x00\x01\x87\xbe\xcb4\xce\
 \x00\x00\x01P\x00\x00\x00\x00\x00\x01\x00\x00\x1d\xef\
-\x00\x00\x01\x8e5\x01\xe3\x9d\
+\x00\x00\x01\x87\xbe\xcb4\xcb\
 \x00\x00\x00\xae\x00\x00\x00\x00\x00\x01\x00\x00\x1b\x95\
-\x00\x00\x01\x8e5\x01\xe4}\
-\x00\x00\x0db\x00\x00\x00\x00\x00\x01\x00\x01nU\
-\x00\x00\x01\x8e5\x01\xe4\x84\
-\x00\x00\x07@\x00\x00\x00\x00\x00\x01\x00\x00\x8aR\
-\x00\x00\x01\x8e5\x01\xe3\xb5\
+\x00\x00\x01\x87\xbe\xcb5e\
+\x00\x00\x0d\x92\x00\x00\x00\x00\x00\x01\x00\x01~1\
+\x00\x00\x01\x87\xbe\xcb5o\
+\x00\x00\x07p\x00\x00\x00\x00\x00\x01\x00\x00\x9a.\
+\x00\x00\x01\x87\xbe\xcb4\xd6\
 \x00\x00\x01p\x00\x00\x00\x00\x00\x01\x00\x00 J\
-\x00\x00\x01\x8e5\x01\xe3\xe1\
-\x00\x00\x07\x1c\x00\x00\x00\x00\x00\x01\x00\x00\x83.\
-\x00\x00\x01\x8e5\x01\xe4\x95\
+\x00\x00\x01\x87\xbe\xcb4\xe8\
+\x00\x00\x07L\x00\x00\x00\x00\x00\x01\x00\x00\x93\x0a\
+\x00\x00\x01\x87\xbe\xcb5}\
 \x00\x00\x00.\x00\x01\x00\x00\x00\x01\x00\x00\x00\x00\
-\x00\x00\x01\x8e5\x01\xe4\xce\
+\x00\x00\x01\x87\xbe\xcb5\xc4\
 "
 
 def qInitResources():
diff --git a/Engine/Source/Programs/SwitchboardListener/SblCore/Private/SwitchboardAuth.cpp b/Engine/Source/Programs/SwitchboardListener/SblCore/Private/SwitchboardAuth.cpp
--- a/Engine/Source/Programs/SwitchboardListener/SblCore/Private/SwitchboardAuth.cpp
+++ b/Engine/Source/Programs/SwitchboardListener/SblCore/Private/SwitchboardAuth.cpp
@@ -481,9 +481,14 @@ namespace UE::SwitchboardListener::Private
 		X509_set_pubkey(x509, InPrivateKey);
 
 		// Set the common name, and copy the subject name to the issuer name.
+		const FUtf8String SubjectCommonName = FUtf8String::Printf(
+			"SwitchboardListener (%s@%s)",
+			TCHAR_TO_UTF8(FPlatformProcess::UserName()),
+			TCHAR_TO_UTF8(FPlatformProcess::ComputerName()));
+
 		X509_NAME* SubjectName = X509_get_subject_name(x509);
-		X509_NAME_add_entry_by_txt(SubjectName, "CN", MBSTRING_ASC,
-			(unsigned char*)"SwitchboardListener", -1, -1, 0);
+		X509_NAME_add_entry_by_txt(SubjectName, "CN", MBSTRING_UTF8,
+			(const unsigned char*)*SubjectCommonName, -1, -1, 0);
 		X509_set_issuer_name(x509, SubjectName);
 
 		// Actually sign the certificate with our key.
