Switchboard: Fix for regression when resolving asset path for project content.

diff --git a/Engine/Plugins/VirtualProduction/Switchboard/Source/Switchboard/switchboard/ue_plugin_utils.py b/Engine/Plugins/VirtualProduction/Switchboard/Source/Switchboard/switchboard/ue_plugin_utils.py
--- a/Engine/Plugins/VirtualProduction/Switchboard/Source/Switchboard/switchboard/ue_plugin_utils.py
+++ b/Engine/Plugins/VirtualProduction/Switchboard/Source/Switchboard/switchboard/ue_plugin_utils.py
@@ -177,7 +177,7 @@ class UnrealPluginManager:
             if (dot_idx := rest.name.find('.')) != -1:
                 rest = rest.parent / rest.name[:dot_idx]
 
-        if mount_name == PROJECT_CONTENT_MOUNT_POINT:
+        if mount_name == PROJECT_CONTENT_MOUNT_NAME:
             assert self._uproject_path
             return self._uproject_path.parent / 'Content' / rest
         else:
