SB: Easier way to set up packaged game launches for nDisplay.

*A dded a button to clear the log window.
* Fixed a couple of issues regarding the values of checkbox change callbacks being of int type instead of QtCore.Qt.CheckState, that were preventing auto scrolling and line wrapping from working.
* Fix issue of Fill DDC not checking if the device is actually an Unreal device, which would cause an exception if any other device was connected, since it would not contain the fill_derived_data_cache function.
* Fixed new config issue where it was not initializing the plugin manager and relatedly keeping the old configuration uproject path and hence could not resolve game paths of maps and possibly other issues as well.
* Should be able to retrieve logs and traces (if the packaged build generates them). Will still put them under the local project path - does not handle the case of no local project available.

#rn Better support for launching packaged nDisplay clusters.

diff --git a/Engine/Plugins/VirtualProduction/Switchboard/Source/Switchboard/switchboard/config.py b/Engine/Plugins/VirtualProduction/Switchboard/Source/Switchboard/switchboard/config.py
--- a/Engine/Plugins/VirtualProduction/Switchboard/Source/Switchboard/switchboard/config.py
+++ b/Engine/Plugins/VirtualProduction/Switchboard/Source/Switchboard/switchboard/config.py
@@ -691,6 +691,16 @@
             'class directly. A derived class (e.g. DirectoryPathSetting or '
             'FilePathSetting) must be used instead.')
 
+    def _getStartPath(self) -> str:
+        ''' Returns a reasonable start path for the file dialog'''
+
+        start_path = str(pathlib.Path.home())
+
+        if (SETTINGS.LAST_BROWSED_PATH and os.path.exists(SETTINGS.LAST_BROWSED_PATH)):
+            start_path = SETTINGS.LAST_BROWSED_PATH
+
+        return start_path
+
     def _create_widgets(
             self, override_device_name: Optional[str] = None) \
             -> QtWidgets.QHBoxLayout:
@@ -705,10 +715,7 @@
         edit_layout.addWidget(browse_btn)
 
         def on_browse_clicked():
-            start_path = str(pathlib.Path.home())
-            if (SETTINGS.LAST_BROWSED_PATH and
-                    os.path.exists(SETTINGS.LAST_BROWSED_PATH)):
-                start_path = SETTINGS.LAST_BROWSED_PATH
+            start_path = self._getStartPath()
 
             fs_path = self._getFileSystemPath(
                 parent=browse_btn, start_path=start_path)
@@ -774,13 +781,31 @@
 
         self.file_path_filter = file_path_filter
 
+    #~ Begin FileSystemPathSetting Interface
+
+    def _getStartPath(self) -> str:
+        ''' Tries to use the current value to initialize the starting path,
+        otherwise returns the base class implementation'''
+
+        current_path = self.get_value()
+
+        if os.path.exists(current_path):
+            return current_path
+
+        return super()._getStartPath()
+
     def _getFileSystemPath(
             self, parent: Optional[QtWidgets.QWidget] = None,
-            start_path: str = '') -> str:
+            start_path: str = ''
+            ) -> str:
+
         file_path, _ = QtWidgets.QFileDialog.getOpenFileName(
             parent=parent, dir=start_path, filter=self.file_path_filter)
+
         return file_path
 
+    #~ End FileSystemPathSetting Interface
 
+
 class PerforcePathSetting(StringSetting):
     '''
@@ -2132,22 +2157,33 @@
         self.SWITCHBOARD_DIR = os.path.abspath(
             os.path.join(os.path.dirname(os.path.abspath(__file__)), '../'))
 
-        # MISC SETTINGS
         self.CURRENT_LEVEL = data.get('current_level', DEFAULT_MAP_TEXT)
 
-        # UE plugin tracking
-        self.ue_plugin_mgr = UnrealPluginManager()
-        self.ue_plugin_mgr.set_engine_dir(Path(self.ENGINE_DIR.get_value()))
-        self.ue_plugin_mgr.set_uproject_path(
-            Path(self.UPROJECT_PATH.get_value()))
+        self.init_plugin_tracking()
+        self.init_devices(data)
 
-        self.ENGINE_DIR.signal_setting_changed.connect(
-            lambda _, new: self.ue_plugin_mgr.set_engine_dir(Path(new)))
+    def _backup_corrupted_config(self, original_file_path: str):
+        directory_name = os.path.dirname(original_file_path)
+        original_file_name = os.path.basename(original_file_path)
+        new_file_name = original_file_name.replace(".", "_corrupted_backup.")
 
-        self.UPROJECT_PATH.signal_setting_changed.connect(
-            lambda _, new: self.ue_plugin_mgr.set_uproject_path(Path(new)))
+        LOGGER.error(f'{original_file_name} has invalid JSON format. Creating default...')
+        answer = QtWidgets.QMessageBox.question(
+            None,
+            'Invalid project settings',
+            f'Config file { original_file_name } is invalid JSON and will be replaced by a new default JSON config.'
+            f'\n\nDo you want to save a backup named { new_file_name }?',
+            QtWidgets.QMessageBox.Yes | QtWidgets.QMessageBox.No
+        )
+        if answer == QtWidgets.QMessageBox.Yes:
+            new_file_path = os.path.join(directory_name, new_file_name)
+            if os.path.exists(new_file_path):
+                os.remove(new_file_path)
+            shutil.copy(original_file_path, new_file_path)
 
-        # Devices
+    def init_devices(self, data={}):
+        ''' '''
+
         self._device_data_from_config = {}
         self._plugin_data_from_config = {}
         self._device_settings = {}
@@ -2176,48 +2212,39 @@
                     self._device_data_from_config.setdefault(
                         device_type, []).append(device_data)
 
-    def _backup_corrupted_config(self, original_file_path: str):
-        directory_name = os.path.dirname(original_file_path)
-        original_file_name = os.path.basename(original_file_path)
-        new_file_name = original_file_name.replace(".", "_corrupted_backup.")
+    def init_plugin_tracking(self):
+        ''' Initializes the plugin manager '''
 
-        LOGGER.error(f'{original_file_name} has invalid JSON format. Creating default...')
-        answer = QtWidgets.QMessageBox.question(
-            None,
-            'Invalid project settings',
-            f'Config file { original_file_name } is invalid JSON and will be replaced by a new default JSON config.'
-            f'\n\nDo you want to save a backup named { new_file_name }?',
-            QtWidgets.QMessageBox.Yes | QtWidgets.QMessageBox.No
-        )
-        if answer == QtWidgets.QMessageBox.Yes:
-            new_file_path = os.path.join(directory_name, new_file_name)
-            if os.path.exists(new_file_path):
-                os.remove(new_file_path)
-            shutil.copy(original_file_path, new_file_path)
+        self.ue_plugin_mgr = UnrealPluginManager()
+        self.ue_plugin_mgr.set_engine_dir(Path(self.ENGINE_DIR.get_value()))
+        self.ue_plugin_mgr.set_uproject_path(Path(self.UPROJECT_PATH.get_value()))
 
+        self.ENGINE_DIR.signal_setting_changed.connect(
+            lambda _, new: self.ue_plugin_mgr.set_engine_dir(Path(new)))
+
+        self.UPROJECT_PATH.signal_setting_changed.connect(
+            lambda _, new: self.ue_plugin_mgr.set_uproject_path(Path(new)))
+
     def init_new_config(self, file_path: Union[str, pathlib.Path], uproject, engine_dir, p4_settings):
-        ''' 
-        Initialize new configuration
-        '''
+        ''' Initialize new configuration '''
 
+        basic_project_settings = {
+                "project_name": self.file_path.stem,
+                "uproject": uproject,
+                "engine_dir": engine_dir,
+            }
+
         self.file_path = get_absolute_config_path(file_path)
         self.init_switchboard_settings()
         self.init_sblhelper_settings()
-        self.init_project_settings(
-            { 
-                "project_name": self.file_path.stem, 
-                "uproject": uproject, 
-                "engine_dir": engine_dir,
-            } | p4_settings)
+        self.init_project_settings(basic_project_settings | p4_settings)
         self.init_unreal_insights()
         self.init_muserver()
 
         self.CURRENT_LEVEL = DEFAULT_MAP_TEXT
 
-        self._device_data_from_config = {}
-        self._plugin_data_from_config = {}
-        self._plugin_settings = {}
-        self._device_settings = {}
+        self.init_plugin_tracking()
+        self.init_devices()
 
         LOGGER.info(f"Creating new config saved in {self.file_path}")
         self.save()
@@ -2228,9 +2255,9 @@
     def init_switchboard_settings(self, data={}):
         self.switchboard_settings = {
             "listener_exe": StringSetting(
-                attr_name = "listener_exe",
-                nice_name = "Listener Executable Name",
-                value = data.get('listener_exe', 'SwitchboardListener')
+                attr_name="listener_exe",
+                nice_name="Listener Executable Name",
+                value=data.get('listener_exe', 'SwitchboardListener')
             )
         }
 
@@ -2239,14 +2266,14 @@
     def init_sblhelper_settings(self, data={}):
         self.sblhelper_settings = {
             "sblhelper_exe": StringSetting(
-                attr_name = "sblhelper_exe",
-                nice_name = "Gpu Clocker Executable Name",
-                value = data.get('sblhelper_exe', 'SwitchboardListenerHelper'),
-                tool_tip = "Name of the executable that SwitchboardListener can communicate with to lock Gpu clocks.",
-                show_ui = True if sys.platform in ('win32','linux') else False, # # Gpu Clocker is available in select platforms
+                attr_name="sblhelper_exe",
+                nice_name="Gpu Clocker Executable Name",
+                value=data.get('sblhelper_exe', 'SwitchboardListenerHelper'),
+                tool_tip="Name of the executable that SwitchboardListener can communicate with to lock Gpu clocks.",
+                show_ui=True if sys.platform in ('win32', 'linux') else False  # Gpu Clocker is available in select platforms
             )
         }
-        
+
         self.SBLHELPER_EXE = self.sblhelper_settings["sblhelper_exe"]
 
     def init_project_settings(self, data={}):
diff --git a/Engine/Plugins/VirtualProduction/Switchboard/Source/Switchboard/switchboard/devices/ndisplay/plugin_ndisplay.py b/Engine/Plugins/VirtualProduction/Switchboard/Source/Switchboard/switchboard/devices/ndisplay/plugin_ndisplay.py
--- a/Engine/Plugins/VirtualProduction/Switchboard/Source/Switchboard/switchboard/devices/ndisplay/plugin_ndisplay.py
+++ b/Engine/Plugins/VirtualProduction/Switchboard/Source/Switchboard/switchboard/devices/ndisplay/plugin_ndisplay.py
@@ -47,6 +47,18 @@ class AddnDisplayDialog(AddDeviceDialog):
         self.name_field = None
         self.address_field = None
 
+        # Initialize a grid layout for the form
+        grid_layout = QtWidgets.QGridLayout()
+
+        # Config File row
+
+        lblConfigFile = QtWidgets.QLabel(self, text="Config File")
+
+        # Combobox with config files
+        self.cbConfigs = sb_widgets.SearchableComboBox(self)
+        self.cbConfigs.setSizePolicy(QtWidgets.QSizePolicy.Expanding, QtWidgets.QSizePolicy.Preferred)
+        self.cbConfigs.setEditable(True)  # to allow the user to type the value
+
         # Button to browse for supported config files
         self.btnBrowse = QtWidgets.QPushButton(self, text="Browse")
         self.btnBrowse.clicked.connect(self.on_clicked_btnBrowse)
@@ -55,19 +67,46 @@ class AddnDisplayDialog(AddDeviceDialog):
         self.btnFindConfigs = QtWidgets.QPushButton(self, text="Populate")
         self.btnFindConfigs.clicked.connect(self.on_clicked_btnFindConfigs)
 
-        # Combobox with config files
-        self.cbConfigs = sb_widgets.SearchableComboBox(self)
-        self.cbConfigs.setSizePolicy(
-            QtWidgets.QSizePolicy.Expanding, QtWidgets.QSizePolicy.Preferred)
-        self.cbConfigs.setEditable(True)  # to allow the user to type the value
+        config_file_layout = QtWidgets.QHBoxLayout()
+        config_file_layout.addWidget(self.cbConfigs)
+        config_file_layout.addWidget(self.btnBrowse)
+        config_file_layout.addWidget(self.btnFindConfigs)
+
+        grid_layout.addWidget(lblConfigFile, 0, 0)
+        grid_layout.addLayout(config_file_layout, 0, 1)
+
+        # Packaged Game row
+
+        lblPackagedGame = QtWidgets.QLabel(self, text="Packaged Game")
+        self.chkPackagedGame = QtWidgets.QCheckBox(self)
+        self.chkPackagedGame.stateChanged.connect(self.on_chkPackagedGame_stateChanged)
 
-        # Create layout for the config file selection widgets
-        file_selection_layout = QtWidgets.QHBoxLayout()
-        file_selection_layout.addWidget(self.cbConfigs)
-        file_selection_layout.addWidget(self.btnBrowse)
-        file_selection_layout.addWidget(self.btnFindConfigs)
+        packaged_game_layout = QtWidgets.QHBoxLayout()
+        packaged_game_layout.addWidget(self.chkPackagedGame)
+        packaged_game_layout.setAlignment(QtCore.Qt.AlignLeft)
 
-        self.form_layout.addRow("Config File", file_selection_layout)
+        grid_layout.addWidget(lblPackagedGame, 1, 0)
+        grid_layout.addLayout(packaged_game_layout, 1, 1)
+
+        # Path row (initially hidden)
+
+        self.lblPath = QtWidgets.QLabel(self, text="Packaged Executable")
+        self.lblPath.setVisible(False)
+        self.pathField = QtWidgets.QLineEdit(self)
+        self.pathField.setVisible(False)
+        self.btnPathBrowse = QtWidgets.QPushButton(self, text="...")
+        self.btnPathBrowse.setVisible(False)
+        self.btnPathBrowse.clicked.connect(self.on_clicked_btnPathBrowse)
+
+        path_layout = QtWidgets.QHBoxLayout()
+        path_layout.addWidget(self.pathField)
+        path_layout.addWidget(self.btnPathBrowse)
+
+        grid_layout.addWidget(self.lblPath, 2, 0)
+        grid_layout.addLayout(path_layout, 2, 1)
+
+        # Add the grid layout to the form layout
+        self.form_layout.addRow(grid_layout)
 
         # Add a spacer right before the ok/cancel buttons
         spacer_layout = QtWidgets.QHBoxLayout()
@@ -77,8 +116,7 @@ class AddnDisplayDialog(AddDeviceDialog):
                 QtWidgets.QSizePolicy.Expanding))
         self.form_layout.addRow("", spacer_layout)
 
-        # Find existing nDisplay devices in order to issue a warning about
-        # replacing them.
+        # Find existing nDisplay devices in order to issue a warning about replacing them.
         self.existing_ndisplay_devices = []
 
         for device in existing_devices:
@@ -94,6 +132,32 @@ class AddnDisplayDialog(AddDeviceDialog):
         # populate the config combobox with the items last populated.
         self.recall_config_itemDatas()
 
+    def on_chkPackagedGame_stateChanged(self, state):
+        ''' Called when the user changes the state of the Packaged Game checkbox
+        It updates the visibility of the packaged game path selection accordingly.
+        '''
+
+        is_checked = QtCore.Qt.CheckState(state) == QtCore.Qt.Checked
+
+        self.lblPath.setVisible(is_checked)
+        self.pathField.setVisible(is_checked)
+        self.pathField.setEnabled(is_checked)
+        self.btnPathBrowse.setVisible(is_checked)
+
+    def on_clicked_btnPathBrowse(self):
+        ''' Called to browse for the packaged game executable path.'''
+
+        file_dialog = QtWidgets.QFileDialog(self)
+        file_path, _ = file_dialog.getOpenFileName(
+            self,
+            "Select Executable",
+            "",
+            "All Files (*);;Executables (*.exe);;Scripts (*.bat *.cmd *.sh)"
+        )
+
+        if file_path:
+            self.pathField.setText(os.path.normpath(file_path))
+
     def recall_config_itemDatas(self):
         '''
         Populate the config combobox with the already found configs, avoiding
@@ -135,12 +199,21 @@ class AddnDisplayDialog(AddDeviceDialog):
         config_path = os.path.normpath(config_path)
         return config_path
 
+    def is_packaged_game(self):
+        ''' Returns true if this is a packaged game'''
+        return self.chkPackagedGame.checkState() == QtCore.Qt.Checked
+
+    def packaged_game_path(self):
+        ''' Returns the path to the packaged game'''
+        return self.pathField.text()
+
     def result(self):
         res = super().result()
         if res == QtWidgets.QDialog.Accepted:
             config_path = self.current_config_path()
-            DevicenDisplay.csettings['ndisplay_config_file'].update_value(
-                config_path)
+            DevicenDisplay.csettings['ndisplay_config_file'].update_value(config_path)
+            DevicenDisplay.csettings['is_packaged_game'].update_value(self.is_packaged_game())
+            DevicenDisplay.csettings['packaged_game_path'].update_value(self.packaged_game_path())
 
         return res
 
@@ -297,6 +370,20 @@ class DevicenDisplay(DeviceUnreal):
             allow_reset=False,
             is_read_only=True
         ),
+        'is_packaged_game': BoolSetting(
+            attr_name="is_packaged_game",
+            nice_name="Packaged Game",
+            value=False,
+            tool_tip="Check if launching a packaged game."
+        ),
+        'packaged_game_path': FilePathSetting(
+            attr_name="packaged_game_path",
+            nice_name="Packaged Executable",
+            value="",
+            tool_tip="Path to the nDisplay packaged game executable. Only used when 'Packaged Game' is checked.",
+            show_ui=True,
+            file_path_filter="Programs and Scripts (*.exe;*.bat;*.sh);;All Files (*)"
+        ),
         'use_all_available_cores': BoolSetting(
             attr_name="use_all_available_cores",
             nice_name="Use All Available Cores",
@@ -375,9 +462,9 @@ class DevicenDisplay(DeviceUnreal):
         ),
         'executable_filename': FilePathSetting(
             attr_name="executable_filename",
-            nice_name="nDisplay Executable Filename",
+            nice_name="Unreal Editor Filename",
             value="UnrealEditor.exe",
-            file_path_filter="Programs (*.exe;*.bat)"
+            file_path_filter="Programs (*.exe;*.bat);;All Files (*)"
         ),
         'ndisplay_cmd_args': StringSetting(
             attr_name="ndisplay_cmd_args",
@@ -387,7 +474,7 @@ class DevicenDisplay(DeviceUnreal):
         'ndisplay_exec_cmds': StringListSetting(
             attr_name="ndisplay_exec_cmds",
             nice_name='ExecCmds',
-            value= [],
+            value=[],
             tool_tip='ExecCmds to be passed. No need for outer double quotes.',
             allow_reset=False,
             migrate_data=migrate_comma_separated_string_to_list
@@ -676,6 +763,7 @@ class DevicenDisplay(DeviceUnreal):
             DevicenDisplay.csettings['livelink_preset'],
             DevicenDisplay.csettings['mediaprofile'],
             DevicenDisplay.csettings['graphics_adapter'],
+            DevicenDisplay.csettings['packaged_game_path'],
             CONFIG.ENGINE_DIR,
             CONFIG.SOURCE_CONTROL_WORKSPACE,
             CONFIG.UPROJECT_PATH,
@@ -751,7 +839,7 @@ class DevicenDisplay(DeviceUnreal):
         # single player
         if multiplayer_mode_name == 'None':
             return map_name
-        
+
         def is_local_address(address):
             return address == SETTINGS.ADDRESS.get_value() or address == '127.0.0.1'
 
@@ -767,28 +855,61 @@ class DevicenDisplay(DeviceUnreal):
             return get_client_args(primary_device_address)
 
         if map_name_is_valid(map_name):
-            return '?'.join([map_name,'Listen', get_common_args()])
-         
+            return '?'.join([map_name, 'Listen', get_common_args()])
+
         current_map_name = get_game_launch_level_path()
         if current_map_name.strip() == '':
             return map_name
 
         return '?'.join([current_map_name, 'Listen', get_common_args()])
 
-    def should_use_project_path_in_command_line(self):
+    def should_use_project_path_in_command_line(self) -> bool:
         ''' Returns true if the project path should be added to the command line.
         This is normally true when launching the editor, but not when launching a cooked game executable.
         '''
- 
+
+        # Don't use project path if this is a packaged game
+        if self.is_packaged_game():
+            return False
+
         # The default exe of an Unreal device is the Editor.
         default_exe = Path(DeviceUnreal.csettings['ue_exe']._original_value)
 
         # This is the current executable
-        exe = Path(self.generate_unreal_exe_path())
+        exe = Path(self.generate_exe_path())
 
         # We don't use direct comparison to include editor build variants, such as -Debug builds.
         return exe.stem.lower().startswith(default_exe.stem.lower())
 
+    def is_packaged_game(self) -> bool:
+        ''' Returns True if this node is launching as a package game'''
+        return DevicenDisplay.csettings['is_packaged_game'].get_value()
+
+    def get_packaged_game_path(self) -> str:
+        ''' Returns the packaged game executable path '''
+        return DevicenDisplay.csettings['packaged_game_path'].get_value(self.name)
+
+    def get_remote_log_path(self):
+        ''' Override from base class '''
+
+        # If packaged, assume exe name also the name of the folder where Saved is located in.
+
+        if self.is_packaged_game():
+            exe_path = Path(self.get_packaged_game_path())
+            return exe_path.with_name(exe_path.stem) / 'Saved' / 'Logs'
+
+        return super().get_remote_log_path()
+
+    def generate_exe_path(self) -> str:
+        ''' Uses the packaged executable path when this is a packaged game launch
+        Note: We didn't override generate_unreal_exe_path because Fill DDCs would grab the wrong executable.
+        '''
+
+        if self.is_packaged_game():
+            return self.get_packaged_game_path()
+
+        return super().generate_unreal_exe_path()
+
     def generate_unreal_command_line(self, map_name=""):
 
         uproject = ''
@@ -934,10 +1055,12 @@ class DevicenDisplay(DeviceUnreal):
         # Modify map name arg for multiplayer mode
         map_name = self.generate_multiplayer_map_name_args(map_name)
 
+        dash_game = "-game" if not self.is_packaged_game() else ""
+
         # fill in fixed arguments
         args = [
             f'{uproject}',
-            "-game",                      # render nodes run in -game
+            dash_game,                    # render nodes run in -game
             f'{map_name}',                # map to open
             "-messaging",                 # enables messaging, needed for MultiUser
             "-dc_cluster",                # this is a cluster node
@@ -1087,7 +1210,7 @@ class DevicenDisplay(DeviceUnreal):
         args.append(self.csettings['logging'].get_command_line_arg(
             override_device_name=self.name))
 
-        path_to_exe = self.generate_unreal_exe_path()
+        path_to_exe = self.generate_exe_path()
         args_expanded = ' '.join(args)
 
         self.settings['ue_command_line'].update_value(
@@ -1499,7 +1622,8 @@ class DevicenDisplay(DeviceUnreal):
             cfg_content, cfg_destination)
         self.unreal_client.send_message(cfg_msg)
 
-        if self.bp_object_path:
+        # Send .uasset but only if not launching packaged
+        if self.bp_object_path and not self.is_packaged_game():
             local_project_path = os.path.dirname(
                 CONFIG.UPROJECT_PATH.get_value())
             dest_project_path = os.path.dirname(
diff --git a/Engine/Plugins/VirtualProduction/Switchboard/Source/Switchboard/switchboard/switchboard_dialog.py b/Engine/Plugins/VirtualProduction/Switchboard/Source/Switchboard/switchboard/switchboard_dialog.py
--- a/Engine/Plugins/VirtualProduction/Switchboard/Source/Switchboard/switchboard/switchboard_dialog.py
+++ b/Engine/Plugins/VirtualProduction/Switchboard/Source/Switchboard/switchboard/switchboard_dialog.py
@@ -419,6 +419,7 @@ class SwitchboardDialog(QtCore.QObject):
             lambda: self._set_engine_changelist(self.window.engine_cl_combo_box.currentText()))
         self.window.logger_level_comboBox.currentTextChanged.connect(self.logger_level_comboBox_currentTextChanged)
         self.window.logger_autoscroll_checkbox.stateChanged.connect(self.logger_autoscroll_stateChanged)
+        self.window.clear_log_button.clicked.connect(self.clear_log_button_clicked)
         self.window.logger_wrap_checkbox.stateChanged.connect(self.logger_wrap_stateChanged)
         self.window.record_button.released.connect(self.record_button_released)
         self.window.sync_all_button.clicked.connect(self.sync_all_button_clicked)
@@ -700,7 +701,8 @@ class SwitchboardDialog(QtCore.QObject):
 
     def register_fill_ddc_menuitem(self):
         def fill_ddc(current_level_only):
-            for device in self.device_manager.devices():
+            unrealdevices = [device for device in self.device_manager.devices() if isinstance(device, DeviceUnreal)]
+            for device in unrealdevices:
                 if not device.is_disconnected:
                     device.fill_derived_data_cache(current_level_only)
 
@@ -1188,6 +1190,10 @@ class SwitchboardDialog(QtCore.QObject):
         else:
             return True
 
+    def clear_log_button_clicked(self):
+        ''' Called when the "Clear" button of the log window is clicked'''
+        self.window.base_console.clear()
+
     def sync_all_button_clicked(self):
         if not CONFIG.P4_ENABLED.get_value():
             return
@@ -1918,14 +1924,14 @@ class SwitchboardDialog(QtCore.QObject):
             LOGGER.setLevel(logging.INFO)
 
     def logger_autoscroll_stateChanged(self, value):
-        if value == QtCore.Qt.Checked:
+        if QtCore.Qt.CheckState(value) == QtCore.Qt.Checked:
             self.logger_autoscroll = True
             self.logger_scroll_to_end()
         else:
             self.logger_autoscroll = False
 
     def logger_wrap_stateChanged(self, value):
-        if value == QtCore.Qt.Checked:
+        if QtCore.Qt.CheckState(value) == QtCore.Qt.Checked:
             self.window.base_console.setLineWrapMode(QtWidgets.QPlainTextEdit.LineWrapMode.WidgetWidth)
         else:
             self.window.base_console.setLineWrapMode(QtWidgets.QPlainTextEdit.LineWrapMode.NoWrap)
