SB: Added action menu to package the cluster.

The user is now able to package the game and launch it (for nDisplay devices). Toggling between packaged and -game is a matter of toggling the packaged game checkbox in nDisplay settings.

Also now passing tooltips for the QWidgetActions getting registered.

#rn Cluster can now be packaged from Switchboard.

diff --git a/Engine/Plugins/VirtualProduction/Switchboard/Source/Switchboard/switchboard/devices/ndisplay/plugin_ndisplay.py b/Engine/Plugins/VirtualProduction/Switchboard/Source/Switchboard/switchboard/devices/ndisplay/plugin_ndisplay.py
--- a/Engine/Plugins/VirtualProduction/Switchboard/Source/Switchboard/switchboard/devices/ndisplay/plugin_ndisplay.py
+++ b/Engine/Plugins/VirtualProduction/Switchboard/Source/Switchboard/switchboard/devices/ndisplay/plugin_ndisplay.py
@@ -7,7 +7,9 @@ from pathlib import Path, PurePosixPath
 import socket
 import struct
 import traceback
+import uuid
 from typing import Optional
+from enum import Enum
 
 from PySide6 import QtCore
 from PySide6 import QtWidgets
@@ -24,14 +26,23 @@ from switchboard.devices.unreal.plugin_unreal import DeviceUnreal, \
     DeviceWidgetUnreal, LiveLinkPresetSetting, MediaProfileSetting
 from switchboard.devices.unreal.uassetparser import UassetParser
 from switchboard.devices.device_base import DeviceStatus
-from switchboard.message_protocol import SyncStatusRequestFlags
 from switchboard.switchboard_logging import LOGGER
 from switchboard.sbcache import SBCache, Asset
+from switchboard.devices.unreal.plugin_unreal import ProgramStartQueueItem
 
 from .ndisplay_monitor_ui import nDisplayMonitorUI
 from .ndisplay_monitor import nDisplayMonitor
 
 
+class PackagingClientConfig(Enum):
+    ''' Used to specify the type of client packaging'''
+
+    Default = "Default"
+    NoClient = "None"
+    Shipping = "Shipping"
+    Development = "Development"
+
+
 class AddnDisplayDialog(AddDeviceDialog):
     def __init__(self, existing_devices, parent=None):
         super().__init__(
@@ -89,8 +100,7 @@ class AddnDisplayDialog(AddDeviceDialog):
         grid_layout.addLayout(packaged_game_layout, 1, 1)
 
         # Path row (initially hidden)
-
-        self.lblPath = QtWidgets.QLabel(self, text="Packaged Executable")
+        self.lblPath = QtWidgets.QLabel(self, text=DevicenDisplay.csettings['packaged_game_path'].nice_name)
         self.lblPath.setVisible(False)
         self.pathField = QtWidgets.QLineEdit(self)
         self.pathField.setVisible(False)
@@ -1300,6 +1310,76 @@ class DevicenDisplay(DeviceUnreal):
         ''' Selects this node as the primary node in the nDisplay cluster '''
         self.__class__.select_device_as_primary(self)
 
+    def package_game(self, clientconfig: PackagingClientConfig = PackagingClientConfig.Default) -> uuid.UUID:
+        ''' Packages the game using default settings '''
+
+        program_name = "unreal_packagegame"
+        outdir = Path(self.get_packaged_game_path())
+
+        if outdir.is_file():
+            outdir = outdir.parent
+
+        # @todo using the remote platform would be more correct.
+
+        if sys.platform.startswith('linux'):
+            platform = 'Linux'
+        elif sys.platform.startswith('darwin'):
+            platform = 'Mac'
+        else:
+            platform = 'Win64'
+
+        if clientconfig == PackagingClientConfig.Default:
+            clientconfigarg = ''
+        else:
+            clientconfigarg = f'-clientconfig={clientconfig.value}'
+
+        args = [
+            'BuildCookRun',
+            f'-project="{CONFIG.UPROJECT_PATH.get_value(self.name)}"',
+            '-noP4',
+            f'-platform={platform}',
+            clientconfigarg,
+            '-cook',
+            '-allmaps',
+            '-build',
+            '-stage',
+            '-pack',
+            '-archive',
+            f'-archivedirectory="{str(outdir)}"'
+        ]
+
+        prog_exe = "RunUAT"
+
+        if sys.platform.startswith('win'):
+            prog_exe += ".bat"
+        else:
+            prog_exe += ".sh"
+
+        prog_path = str(Path(CONFIG.ENGINE_DIR.get_value(self.name)).parent / prog_exe)
+        prog_args = ' '.join(args)
+
+        LOGGER.info(f'Packaging "{self.name}" with command: {prog_path} {prog_args}')
+
+        puuid, msg = message_protocol.create_start_process_message(
+            prog_path=prog_path,
+            prog_args=prog_args,
+            prog_name=program_name,
+            caller=self.name,
+            update_clients_with_stdout=True,
+        )
+
+        self.program_start_queue.add(
+            ProgramStartQueueItem(
+                name=program_name,
+                puuid_dependency=None,
+                puuid=puuid,
+                msg_to_unreal_client=msg,
+            ),
+            unreal_client=self.unreal_client,
+        )
+
+        return puuid
+
     @classmethod
     def extract_configexport_from_uasset(cls, cfg_file) -> str:
         ''' Extract the configexport from the config uasset '''
diff --git a/Engine/Plugins/VirtualProduction/Switchboard/Source/Switchboard/switchboard/switchboard_dialog.py b/Engine/Plugins/VirtualProduction/Switchboard/Source/Switchboard/switchboard/switchboard_dialog.py
--- a/Engine/Plugins/VirtualProduction/Switchboard/Source/Switchboard/switchboard/switchboard_dialog.py
+++ b/Engine/Plugins/VirtualProduction/Switchboard/Source/Switchboard/switchboard/switchboard_dialog.py
@@ -43,6 +43,7 @@ from switchboard.tools.insights_launcher import InsightsLauncher
 from switchboard.tools.listener_launcher import ListenerLauncher
 from switchboard.tools.sblhelper_launcher import SBLHelperLauncher
 from switchboard.devices.unreal.plugin_unreal import DeviceUnreal
+from switchboard.devices.ndisplay.plugin_ndisplay import DevicenDisplay, PackagingClientConfig
 from switchboard.devices.unreal.redeploy_dialog import RedeployListenerDialog
 from switchboard.util import collect_logs
 
@@ -219,6 +220,7 @@ class ProcessMonitor(QtCore.QObject):
             self._dialog.update_locallistener_menuitem()
             self._dialog.update_localsblhelper_menuitem()
             self._dialog.update_insights_menuitem()
+            self._dialog.update_package_game_menuitem()
 
             time.sleep(1.0)
 
@@ -333,6 +335,9 @@ class SwitchboardDialog(QtCore.QObject):
         # Fill DDC submenu
         self.register_fill_ddc_menuitem()
 
+        # Package game
+        self.register_package_game_menuitem()
+
         # Transport Manager
         #self.transport_queue = recording.TransportQueue(CONFIG.SWITCHBOARD_DIR)
         #self.transport_queue.signal_transport_queue_job_started.connect(self.transport_queue_job_started)
@@ -611,6 +616,24 @@ class SwitchboardDialog(QtCore.QObject):
         '''
         self.insights_launcher_menuitem.setEnabled(not self.insights_launcher.is_running())
 
+    def update_package_game_menuitem(self):
+        ''' Enables/disables the package game menu depending on whether
+        there are any nDisplay devices connected or not.
+        '''
+        ndisplay_devices = [device for device in self.device_manager.devices()
+                            if isinstance(device, DevicenDisplay) and not device.is_disconnected]
+
+        do_enable = len(ndisplay_devices) > 0
+
+        self.action_package_development.setEnabled(do_enable)
+        self.action_package_shipping.setEnabled(do_enable)
+
+        parent = self.action_package_development.parent()
+        while parent is not None and not isinstance(parent, QMenu):
+            parent = parent.parent()
+
+        parent.setEnabled(do_enable)
+
     def on_muserver_start_stop_click(self):
         '''
         Handle the multi-user server button click. If we are running we stop the process. If we are not
@@ -717,12 +740,45 @@ class SwitchboardDialog(QtCore.QObject):
         all_levels_action = self.register_tools_menu_action("All Levels", ["Fill DDC (Prepare Shaders)"])
         all_levels_action.triggered.connect(fill_ddc_all_levels_action)
 
+    def register_package_game_menuitem(self):
+        ''' Registers a menu item to package the nDisplay cluster game'''
+
+        def package_nDisplay_game(clientconfig: PackagingClientConfig):
+            ''' Callback to kick off the packaging process '''
+            ndisplay_devices = [device for device in self.device_manager.devices() 
+                                if isinstance(device, DevicenDisplay) and not device.is_disconnected]
+
+            for device in ndisplay_devices:
+                device.package_game(clientconfig)
+
+        actionname = "Package nDisplay Game"
+
+        actiontooltip = '\n'.join([
+            "Packages the project in all the nDisplay nodes currently connected to.",
+            f"The archive will reside in the location specified by '{DevicenDisplay.csettings['packaged_game_path'].nice_name}'",
+            "The Engine and Project source must be available on the remote machine."
+            ])
+
+        self.action_package_development = self.register_tools_menu_action(
+            "Development", menunames=[actionname], actiontooltip=actiontooltip
+        )
+        self.action_package_development.triggered.connect(lambda: package_nDisplay_game(PackagingClientConfig.Development))
+
+        self.action_package_shipping = self.register_tools_menu_action(
+            "Shipping", menunames=[actionname], actiontooltip=actiontooltip
+        )
+        self.action_package_shipping.triggered.connect(lambda: package_nDisplay_game(PackagingClientConfig.Shipping))
+
     def add_tools_menu(self):
         ''' Adds tools menu to menu bar and populates built-in items '''
 
         self.tools_menu = self.window.menu_bar.addMenu("&Tools")
 
-    def register_tools_menu_action(self, actionname: str, menunames: List[str] = []) -> QWidgetAction:
+    def register_tools_menu_action(
+            self,
+            actionname: str,
+            menunames: List[str] = [],
+            actiontooltip: str = '') -> QWidgetAction:
         ''' Registers a QWidgetAction with the tools menu
 
         Args:
@@ -762,6 +818,8 @@ class SwitchboardDialog(QtCore.QObject):
         # add the given action
         action = QWidgetAction(current_menu)
         action.setText(actionname)
+        action.setToolTip(actiontooltip)
+        current_menu.setToolTipsVisible(True)
         current_menu.addAction(action)
 
         return action
