SB: Added categories to Settings for organizational purposes.

Search works with this new widget hierarchy.

#rnx

diff --git a/Engine/Plugins/VirtualProduction/Switchboard/Source/Switchboard/switchboard/config.py b/Engine/Plugins/VirtualProduction/Switchboard/Source/Switchboard/switchboard/config.py
--- a/Engine/Plugins/VirtualProduction/Switchboard/Source/Switchboard/switchboard/config.py
+++ b/Engine/Plugins/VirtualProduction/Switchboard/Source/Switchboard/switchboard/config.py
@@ -121,21 +121,26 @@ class Setting(QtCore.QObject):
         show_ui: bool = True,
         allow_reset: bool = True,
         migrate_data: Optional[Callable[[Any], None]] = None,
+        category: str = 'Misc'
     ):
         '''
         Create a new Setting object.
 
         Args:
-            attr_name: Internal name.
-            nice_name: Display name.
-            value    : The initial value of this Setting.
-            tool_tip : Tooltip to show in the UI for this Setting.
-            show_ui  : Whether to show this Setting in the Settings UI.
+            attr_name   : Internal name.
+            nice_name   : Display name.
+            value       : The initial value of this Setting.
+            tool_tip    : Tooltip to show in the UI for this Setting.
+            show_ui     : Whether to show this Setting in the Settings UI.
+            allow_reset : Allows showing a reset button when the value differs from the default.
+            migrate_data: Optional function to migrate data already stored when value structure changes.
+            category    : Used for UI grouping with other properties.
         '''
         super().__init__()
 
         self.attr_name = attr_name
         self.nice_name = nice_name
+        self.category = category
 
         value = self._filter_value(value)
         self._original_value = self._value = value
@@ -365,16 +370,18 @@ class Setting(QtCore.QObject):
         return top_level_widget
 
     def _register_on_setting_changed(self, top_level_widget: QtWidgets.QWidget, override_device_name: str):
-        on_setting_changed_lambda = lambda old_value, new_value, override_device_name=override_device_name: \
+
+        def on_setting_changed(old_value, new_value, override_device_name=override_device_name):
             self._on_setting_changed(new_value, override_device_name=override_device_name)
-        self.signal_setting_changed.connect(
-            on_setting_changed_lambda
-        )
+
+        self.signal_setting_changed.connect(on_setting_changed)
 
         # Clear the widget when it is destroyed to avoid dangling references
-        top_level_widget.destroyed.connect(lambda destroyed_object=None:
-            self._on_widget_destroyed(on_setting_changed_lambda, override_device_name)
-        )
+
+        def handle_widget_destroyed(destroyed_object=None):
+            self._on_widget_destroyed(on_setting_changed, override_device_name)
+
+        top_level_widget.destroyed.connect(handle_widget_destroyed)
 
     def _on_widget_destroyed(self, on_setting_changed_lambda, override_device_name: str):
         try:
@@ -532,29 +539,17 @@ class IntSetting(Setting):
 
     def __init__(
         self,
-        attr_name: str,
-        nice_name: str,
-        value: str,
-        tool_tip: Optional[str] = None,
-        show_ui: bool = True,
-        allow_reset: bool = True,
-        migrate_data: Optional[Callable[[Any], None]] = None,
-        is_read_only: bool = False
+        *args,
+        is_read_only: bool = False,
+        **kwargs
     ):
         '''
         Create a new IntSetting object.
 
         Args:
-            attr_name       : Internal name.
-            nice_name       : Display name.
-            value           : The initial value of this Setting.
-            tool_tip        : Tooltip to show in the UI for this Setting.
-            show_ui         : Whether to show this Setting in the Settings UI.
             is_read_only    : Whether to make entry field editable or not.
         '''
-        super().__init__(
-            attr_name, nice_name, value,
-            tool_tip=tool_tip, show_ui=show_ui, allow_reset=allow_reset, migrate_data=migrate_data)
+        super().__init__(*args, **kwargs)
 
         self.is_read_only = is_read_only
 
@@ -603,31 +598,19 @@ class StringSetting(Setting):
 
     def __init__(
         self,
-        attr_name: str,
-        nice_name: str,
-        value: str,
+        *args,
         placeholder_text: str = '',
-        tool_tip: Optional[str] = None,
-        show_ui: bool = True,
-        allow_reset: bool = True,
-        migrate_data: Optional[Callable[[Any], None]] = None,
-        is_read_only: bool = False
+        is_read_only: bool = False,
+        **kwargs
     ):
         '''
         Create a new StringSetting object.
 
         Args:
-            attr_name       : Internal name.
-            nice_name       : Display name.
-            value           : The initial value of this Setting.
             placeholder_text: Placeholder for this Setting's value in the UI.
-            tool_tip        : Tooltip to show in the UI for this Setting.
-            show_ui         : Whether to show this Setting in the Settings UI.
             is_read_only    : Whether to make entry field editable or not.
         '''
-        super().__init__(
-            attr_name, nice_name, value,
-            tool_tip=tool_tip, show_ui=show_ui, allow_reset=allow_reset, migrate_data=migrate_data)
+        super().__init__(*args, **kwargs)
 
         self.placeholder_text = placeholder_text
         self.is_read_only = is_read_only
@@ -755,37 +738,22 @@ class FilePathSetting(FileSystemPathSetting):
 
     def __init__(
         self,
-        attr_name: str,
-        nice_name: str,
-        value: str,
-        placeholder_text: str = '',
+        *args,
         file_path_filter: str = '',
-        tool_tip: Optional[str] = None,
-        show_ui: bool = True
+        **kwargs
     ):
         '''
         Create a new FilePathSetting object.
 
         Args:
-            attr_name       : Internal name.
-            nice_name       : Display name.
-            value           : The initial value of this Setting.
-            placeholder_text: Placeholder for this Setting's value in the UI.
             file_path_filter: Filter to use in the file browser.
-            tool_tip        : Tooltip to show in the UI for this Setting.
-            show_ui         : Whether to show this Setting in the Settings UI.
         '''
-        super().__init__(
-            attr_name, nice_name, value, placeholder_text=placeholder_text,
-            tool_tip=tool_tip, show_ui=show_ui)
+        super().__init__(*args, **kwargs)
 
         self.file_path_filter = file_path_filter
 
-    #~ Begin FileSystemPathSetting Interface
-
     def _getStartPath(self) -> str:
-        ''' Tries to use the current value to initialize the starting path,
-        otherwise returns the base class implementation'''
+        ''' Override from base class'''
 
         current_path = self.get_value()
 
@@ -798,14 +766,13 @@ class FilePathSetting(FileSystemPathSetting):
             self, parent: Optional[QtWidgets.QWidget] = None,
             start_path: str = ''
             ) -> str:
+        ''' Override from base class'''
 
         file_path, _ = QtWidgets.QFileDialog.getOpenFileName(
             parent=parent, dir=start_path, filter=self.file_path_filter)
 
         return file_path
 
-    #~ End FileSystemPathSetting Interface
-
 
 class PerforcePathSetting(StringSetting):
     '''
@@ -815,23 +782,15 @@ class PerforcePathSetting(StringSetting):
 
     def __init__(
         self,
-        attr_name: str,
-        nice_name: str,
-        value: str,
+        *args,
         placeholder_text: str = '',
-        tool_tip: Optional[str] = None,
-        show_ui: bool = True,
-        allow_reset: bool = True,
-        migrate_data: Optional[Callable[[Any], None]] = None,
         is_read_only: bool = False,
+        **kwargs
     ):
         # Trim matching file paths to the parent directory (e.g. ['.uproject'])
         self.truncate_files_with_extensions: list[str] = []
 
-        super().__init__(
-            attr_name, nice_name, value, placeholder_text=placeholder_text,
-            tool_tip=tool_tip, show_ui=show_ui, allow_reset=allow_reset,
-            migrate_data=migrate_data, is_read_only=is_read_only)
+        super().__init__(*args, **kwargs)
 
     def _filter_value(self, value: Optional[str]) -> str:
         '''
@@ -860,30 +819,17 @@ class OptionSetting(Setting):
 
     def __init__(
         self,
-        attr_name: str,
-        nice_name: str,
-        value,
+        *args,
         possible_values: Optional[list] = None,
-        tool_tip: Optional[str] = None,
-        show_ui: bool = True,
-        allow_reset: bool = True,
-        migrate_data: Optional[Callable[[Any], None]] = None,
+        **kwargs
     ):
         '''
         Create a new OptionSetting object.
 
         Args:
-            attr_name      : Internal name.
-            nice_name      : Display name.
-            value          : The initial value of this Setting.
             possible_values: Possible values for this Setting.
-            tool_tip       : Tooltip to show in the UI for this Setting.
-            show_ui        : Whether to show this Setting in the Settings UI.
         '''
-        super().__init__(
-            attr_name, nice_name, value,
-            tool_tip=tool_tip, show_ui=show_ui, allow_reset=allow_reset,
-            migrate_data=migrate_data)
+        super().__init__(*args, **kwargs)
 
         self.possible_values = possible_values or []
 
@@ -1051,27 +997,16 @@ class ListSetting(Setting):
 
     def __init__(
         self,
-        attr_name: str,
-        nice_name: str,
-        value: list = [],
-        tool_tip: Optional[str] = None,
-        show_ui: bool = True,
-        allow_reset: bool = True,
-        migrate_data: Optional[Callable[[Any], None]] = None
+        *args,
+        **kwargs
     ):
         '''
         Create a new ListSetting object.
 
         Args:
-            attr_name      : Internal name.
-            nice_name      : Display name.
-            value          : The initial value of this Setting.
-            tool_tip       : Tooltip to show in the UI for this Setting.
-            show_ui        : Whether to show this Setting in the Settings UI.
+            value : list
         '''
-        super().__init__(
-            attr_name, nice_name, value,
-            tool_tip=tool_tip, show_ui=show_ui, allow_reset=allow_reset, migrate_data=migrate_data)
+        super().__init__(*args, **kwargs)
 
         self.array_count_labels = {}
         self.element_layouts = {}
@@ -1260,30 +1195,6 @@ class StringListSetting(ListSetting):
     '''
     An array setting where the elements are strings
     '''
-    def __init__(
-        self,
-        attr_name: str,
-        nice_name: str,
-        value: list[str] = [],
-        tool_tip: Optional[str] = None,
-        show_ui: bool = True,
-        allow_reset: bool = True,
-        migrate_data: Optional[Callable[[Any], None]] = None
-    ):
-        '''
-        Create a new ArraySetting object.
-
-        Args:
-            attr_name      : Internal name.
-            nice_name      : Display name.
-            value          : The initial value of this Setting.
-            tool_tip       : Tooltip to show in the UI for this Setting.
-            show_ui        : Whether to show this Setting in the Settings UI.
-        '''
-        super().__init__(
-            attr_name, nice_name, value,
-            tool_tip=tool_tip, show_ui=show_ui, allow_reset=allow_reset, migrate_data=migrate_data)
-        pass
 
     def create_element(self, override_device_name: str, index: int) -> Tuple[QtWidgets.QWidget, object]:
         line_edit = QtWidgets.QLineEdit()
@@ -1635,7 +1546,7 @@ class LoggingSetting(Setting):
     def _filter_value(
             self,
             value: Optional[dict[str, Optional[str]]]
-    )-> collections.OrderedDict:
+    ) -> collections.OrderedDict:
         '''
         Filter function to modify the incoming value before updating or
         overriding the setting.
@@ -1666,37 +1577,26 @@ class LoggingSetting(Setting):
 
     def __init__(
         self,
-        attr_name: str,
-        nice_name: str,
-        value: dict[str, Optional[str]],
+        *args,
         categories: Optional[list[str]] = None,
         verbosity_levels: Optional[list[str]] = None,
-        tool_tip: Optional[str] = None,
-        show_ui: bool = True,
-        allow_reset: bool = True,
-        migrate_data: Optional[Callable[[Any], None]] = None
+        **kwargs
     ):
         '''
         Create a new LoggingSetting object.
 
         Args:
-            attr_name       : Internal name.
-            nice_name       : Display name.
-            value           : The initial value of this Setting.
+            value           : dict[str, Optional[str]]
             categories      : The initial list of logging categories.
             verbosity_levels: The possible settings for verbosity level of
                               each category.
-            tool_tip        : Tooltip to show in the UI for this Setting.
-            show_ui         : Whether to show this Setting in the Settings UI.
         '''
 
         # Set the categories before calling the base class init since they
         # will be used when filtering the value.
         self._categories = categories or []
 
-        super().__init__(
-            attr_name, nice_name, value,
-            tool_tip=tool_tip, show_ui=show_ui, allow_reset=allow_reset, migrate_data=migrate_data)
+        super().__init__(*args, **kwargs)
 
         self._verbosity_levels = (
             verbosity_levels or self.DEFAULT_VERBOSITY_LEVELS)
@@ -1859,25 +1759,14 @@ class LoggingSetting(Setting):
 
 
 class AddressSetting(OptionSetting):
-    def __init__(
-        self,
-        attr_name,
-        nice_name,
-        value,
-        tool_tip=None,
-        show_ui=True,
-        allow_reset=True,
-        migrate_data=None
-    ):
+
+    def __init__(self, *args, **kwargs):
+
         super().__init__(
-            attr_name=attr_name,
-            nice_name=nice_name,
-            value=value,
+            *args,
             possible_values=list(self.generate_possible_addresses()),
-            tool_tip=tool_tip,
-            show_ui=show_ui,
-            allow_reset=allow_reset,
-            migrate_data=migrate_data)
+            **kwargs
+        )
 
     def _create_widgets(
         self, override_device_name: Optional[str] = None
@@ -2281,38 +2170,44 @@ class Config(object):
             "project_name": StringSetting(
                 "project_name",
                 "Project Name",
-                data.get('project_name', 'Default')
+                data.get('project_name', 'Default'),
+                category="General Settings",
             ),
             "uproject": FilePathSetting(
                 "uproject", "uProject Path",
                 data.get('uproject', ''),
-                tool_tip="Path to uProject"
+                tool_tip="Path to uProject",
+                category="General Settings",
             ),
             "engine_dir": DirectoryPathSetting(
                 "engine_dir",
                 "Engine Directory",
                 data.get('engine_dir', ''),
-                tool_tip="Path to UE 'Engine' directory"
+                tool_tip="Path to UE 'Engine' directory",
+                category="General Settings",
             ),
             'engine_sync_method': OptionSetting(
                 "engine_sync_method",
                 "Engine Sync Method",
                 EngineSyncMethod.Use_Existing.value,
                 possible_values=[p.value for p in EngineSyncMethod],
+                category="Source Control Settings",
             ),
             "maps_path": StringSetting(
                 "maps_path",
                 "Map Path",
                 data.get('maps_path', ''),
                 placeholder_text="Maps",
-                tool_tip="Relative path from Content folder that contains maps to launch into."
+                tool_tip="Relative path from Content folder that contains maps to launch into.",
+                category="General Settings",
             ),
             "maps_filter": StringSetting(
                 "maps_filter",
                 "Map Filter",
                 data.get('maps_filter', '*.umap'),
                 placeholder_text="*.umap",
-                tool_tip="Walk every file in the Map Path and run a fnmatch to filter the file names"
+                tool_tip="Walk every file in the Map Path and run a fnmatch to filter the file names",
+                category="General Settings",
             ),
             'maps_plugin_filters': StringListSetting(
                 "maps_plugin_filters",
@@ -2323,7 +2218,8 @@ class Config(object):
                     "Plugins whose name matches any of these filters will "
                     "also be searched for maps."),
                 show_ui=False,
-                migrate_data=migrate_comma_separated_string_to_list
+                migrate_data=migrate_comma_separated_string_to_list,
+                category="General Settings",
             ),
             'content_plugin_filters': StringListSetting(
                 "content_plugin_filters",
@@ -2341,7 +2237,8 @@ class Config(object):
                     "absolute path to a plugin directory. Relative paths "
                     "should be relative to the directory containing the "
                     ".uproject file."),
-                migrate_data=migrate_comma_separated_string_to_list
+                migrate_data=migrate_comma_separated_string_to_list,
+                category="General Settings",
             ),
         }
 
@@ -2386,24 +2283,28 @@ class Config(object):
                 "p4_enabled",
                 "Perforce Enabled",
                 data.get("p4_enabled", False),
-                tool_tip="Toggle Perforce support for the entire application"
+                tool_tip="Toggle Perforce support for the entire application",
+                category="Source Control Settings",
             ),
             "source_control_workspace": StringSetting(
                 "source_control_workspace", "Workspace Name",
                 data.get("source_control_workspace"),
-                tool_tip="SourceControl Workspace/Branch"
+                tool_tip="SourceControl Workspace/Branch",
+                category="Source Control Settings",
             ),
             "p4_sync_path": PerforcePathSetting(
                 "p4_sync_path",
                 "Perforce Project Path",
                 data.get("p4_sync_path", ''),
-                placeholder_text="//UE/Project"
+                placeholder_text="//UE/Project",
+                category="Source Control Settings",
             ),
             "p4_engine_path": PerforcePathSetting(
                 "p4_engine_path",
                 "Perforce Engine Path",
                 data.get("p4_engine_path", ''),
-                placeholder_text="//UE/Project/Engine"
+                placeholder_text="//UE/Project/Engine",
+                category="Source Control Settings",
             )
         }
 
diff --git a/Engine/Plugins/VirtualProduction/Switchboard/Source/Switchboard/switchboard/devices/device_base.py b/Engine/Plugins/VirtualProduction/Switchboard/Source/Switchboard/switchboard/devices/device_base.py
--- a/Engine/Plugins/VirtualProduction/Switchboard/Source/Switchboard/switchboard/devices/device_base.py
+++ b/Engine/Plugins/VirtualProduction/Switchboard/Source/Switchboard/switchboard/devices/device_base.py
@@ -7,7 +7,7 @@ import os
 import random
 import socket
 import threading
-from typing import Optional, Type, TYPE_CHECKING
+from typing import Optional, Type, Dict, OrderedDict, TYPE_CHECKING
 
 from PySide6 import QtCore
 from PySide6 import QtGui
@@ -74,6 +74,7 @@ class Device(QtCore.QObject):
             nice_name='Is Recording Device',
             value=True,
             tool_tip='Is this device used to record',
+            category='Device Settings',
         )
     }
 
@@ -83,7 +84,7 @@ class Device(QtCore.QObject):
         self._name = name  # Assigned name
         self.device_qt_handler = DeviceQtHandler()
         self.autojoin_mu_server = False
-        self.setting_address = AddressSetting('address', 'Address', address)
+        self.setting_address = AddressSetting('address', 'Address', address, category="Network Settings")
         self.setting_address.signal_setting_changed.connect(
             lambda: self._check_recreate_osc_client())
 
@@ -437,3 +438,8 @@ class Device(QtCore.QObject):
         ''' Called after the device manager adds all the devices of this type
         '''
         pass
+
+    @classmethod
+    def sort_setting_categories(cls, categories: Dict) -> OrderedDict:
+        ''' Sort the settings categories so that they are shown in the desired order '''
+        return OrderedDict(categories)
diff --git a/Engine/Plugins/VirtualProduction/Switchboard/Source/Switchboard/switchboard/devices/device_manager.py b/Engine/Plugins/VirtualProduction/Switchboard/Source/Switchboard/switchboard/devices/device_manager.py
--- a/Engine/Plugins/VirtualProduction/Switchboard/Source/Switchboard/switchboard/devices/device_manager.py
+++ b/Engine/Plugins/VirtualProduction/Switchboard/Source/Switchboard/switchboard/devices/device_manager.py
@@ -161,7 +161,7 @@ class DeviceManager(QtCore.QObject):
             config.load_plugin_settings(name, plugin_settings)
 
     def available_device_plugins(self):
-        return self._plugins.keys()
+        return self._plugins
 
     def plugin_settings(self, device_type):
         return self._plugins[device_type].plugin_settings()
diff --git a/Engine/Plugins/VirtualProduction/Switchboard/Source/Switchboard/switchboard/devices/ndisplay/plugin_ndisplay.py b/Engine/Plugins/VirtualProduction/Switchboard/Source/Switchboard/switchboard/devices/ndisplay/plugin_ndisplay.py
--- a/Engine/Plugins/VirtualProduction/Switchboard/Source/Switchboard/switchboard/devices/ndisplay/plugin_ndisplay.py
+++ b/Engine/Plugins/VirtualProduction/Switchboard/Source/Switchboard/switchboard/devices/ndisplay/plugin_ndisplay.py
@@ -8,7 +8,7 @@ import socket
 import struct
 import traceback
 import uuid
-from typing import Optional
+from typing import Optional, Dict, OrderedDict
 from enum import Enum
 
 from PySide6 import QtCore
@@ -378,13 +378,15 @@ class DevicenDisplay(DeviceUnreal):
             value="",
             tool_tip="Path to nDisplay config file",
             allow_reset=False,
-            is_read_only=True
+            is_read_only=True,
+            category="General Settings",
         ),
         'is_packaged_game': BoolSetting(
             attr_name="is_packaged_game",
             nice_name="Packaged Game",
             value=False,
-            tool_tip="Check if launching a packaged game."
+            tool_tip="Check if launching a packaged game.",
+            category="General Settings",
         ),
         'packaged_game_path': FilePathSetting(
             attr_name="packaged_game_path",
@@ -392,28 +394,33 @@ class DevicenDisplay(DeviceUnreal):
             value="",
             tool_tip="Path to the nDisplay packaged game executable. Only used when 'Packaged Game' is checked.",
             show_ui=True,
-            file_path_filter="Programs and Scripts (*.exe;*.bat;*.sh);;All Files (*)"
+            file_path_filter="Programs and Scripts (*.exe;*.bat;*.sh);;All Files (*)",
+            category="General Settings",
         ),
         'use_all_available_cores': BoolSetting(
             attr_name="use_all_available_cores",
             nice_name="Use All Available Cores",
             value=False,
+            category="UE Settings",
         ),
         'texture_streaming': BoolSetting(
             attr_name="texture_streaming",
             nice_name="Texture Streaming",
             value=True,
+            category="Render Settings",
         ),
         'sound': BoolSetting(
             attr_name="sound",
             nice_name="Sound",
             value=False,
+            category="UE Settings",
         ),
         'loading_screen': BoolSetting(
             attr_name="loading_screen",
             nice_name="Loading Screen",
             value=False,
-            tool_tip='When unchecked, will add -NoLoadingScreen to the command line'
+            tool_tip='When unchecked, will add -NoLoadingScreen to the command line',
+            category="UE Settings",
         ),
         'render_api': OptionSetting(
             attr_name="render_api",
@@ -430,32 +437,37 @@ class DevicenDisplay(DeviceUnreal):
                 "vulkan -sm5",
                 "vulkan -sm6"
             ],
+            category="Render Settings",
         ),
         'multiplayer_mode': OptionSetting(
             attr_name="multiplayer_mode",
             nice_name="Multiplayer Server Mode",
             value='None',
-            possible_values=['None', 'Listen server', 'Dedicated server']
+            possible_values=['None', 'Listen server', 'Dedicated server'],
+            category="Multiplayer Settings",
         ),
         'dedicated_server_address': AddressSetting(
             attr_name="dedicated_server_address",
             nice_name="Dedicated Server Address",
             value='127.0.0.1',
             tool_tip='Server address to connect to. Not used if Multiplayer Server Mode '
-                     'is set to "Listen Server" or auto start of dedicated server is enabled'
+                     'is set to "Listen Server" or auto start of dedicated server is enabled',
+            category="Multiplayer Settings",
         ),
         'dedicated_server_port': IntSetting(
             attr_name="dedicated_server_port",
             nice_name="Dedicated Server Port",
             value='7777',
-            tool_tip='Server port to connect to. Not used if Multiplayer Server Mode is not Dedicated Server'
+            tool_tip='Server port to connect to. Not used if Multiplayer Server Mode is not Dedicated Server',
+            category="Multiplayer Settings",
         ),
         'render_mode': OptionSetting(
             attr_name="render_mode",
             nice_name="Render Mode",
             value="Mono",
             possible_values=[
-                "Mono", "Frame sequential", "Side-by-Side", "Top-bottom"]
+                "Mono", "Frame sequential", "Side-by-Side", "Top-bottom"],
+            category="Render Settings",
         ),
         'render_sync_policy': OptionSetting(
             attr_name="render_sync_policy",
@@ -469,17 +481,20 @@ class DevicenDisplay(DeviceUnreal):
                 "- 'Ethernet': Ethernet-based sync. Formerly known as 'sync policy 1'\n"
                 "- 'Nvidia': Nvidia's Quadro Sync Framelock. Formerly known as 'sync policy 2'\n"
             ),
+            category="Render Settings",
         ),
         'executable_filename': FilePathSetting(
             attr_name="executable_filename",
             nice_name="Unreal Editor Filename",
             value="UnrealEditor.exe",
-            file_path_filter="Programs (*.exe;*.bat);;All Files (*)"
+            file_path_filter="Programs (*.exe;*.bat);;All Files (*)",
+            category="UE Settings",
         ),
         'ndisplay_cmd_args': StringSetting(
             attr_name="ndisplay_cmd_args",
             nice_name="Extra Cmd Line Args",
             value="",
+            category="Command Line Args",
         ),
         'ndisplay_exec_cmds': StringListSetting(
             attr_name="ndisplay_exec_cmds",
@@ -487,14 +502,16 @@ class DevicenDisplay(DeviceUnreal):
             value=[],
             tool_tip='ExecCmds to be passed. No need for outer double quotes.',
             allow_reset=False,
-            migrate_data=migrate_comma_separated_string_to_list
+            migrate_data=migrate_comma_separated_string_to_list,
+            category="Command Line Args",
         ),
         'ndisplay_dp_cvars': StringListSetting(
             attr_name='ndisplay_dp_cvars',
             nice_name="DPCVars",
             value=[],
             tool_tip="Device profile console variables.",
-            migrate_data=migrate_comma_separated_string_to_list
+            migrate_data=migrate_comma_separated_string_to_list,
+            category="Command Line Args",
         ),
         'ndisplay_unattended': BoolSetting(
             attr_name='ndisplay_unattended',
@@ -504,6 +521,7 @@ class DevicenDisplay(DeviceUnreal):
                 'Include the "-unattended" command line argument, which is '
                 'documented to "Disable anything requiring feedback from the '
                 'user."'),
+            category="UE Settings",
         ),
         'max_gpu_count': OptionSetting(
             attr_name="max_gpu_count",
@@ -513,13 +531,15 @@ class DevicenDisplay(DeviceUnreal):
             tool_tip=(
                 "If you have multiple GPUs in the PC, you can specify how "
                 "many to use."),
-        ),
+            category="GPU/CPU Settings",
+       ),
         'priority_modifier': OptionSetting(
             attr_name='priority_modifier',
             nice_name="Process Priority",
             value=sb_utils.PriorityModifier.Normal.name,
             possible_values=[p.name for p in sb_utils.PriorityModifier],
             tool_tip="Used to override the priority of the process.",
+            category="GPU/CPU Settings",
         ),
         'populated_config_itemDatas': Setting(
             attr_name='populated_config_itemDatas',
@@ -532,7 +552,8 @@ class DevicenDisplay(DeviceUnreal):
             attr_name='minimize_before_launch',
             nice_name="Minimize Before Launch",
             value=True,
-            tool_tip="Minimizes windows before launch"
+            tool_tip="Minimizes windows before launch",
+            category="General Settings",
         ),
         'primary_device_name': StringSetting(
             attr_name='primary_device_name',
@@ -564,7 +585,8 @@ class DevicenDisplay(DeviceUnreal):
                 'LogLiveLink',
                 'LogRemoteControl',
             ],
-            tool_tip='Logging categories and verbosity levels'
+            tool_tip='Logging categories and verbosity levels',
+            category="UE Settings",
         ),
         'udpmessaging_unicast_endpoint': StringSetting(
             attr_name='udpmessaging_unicast_endpoint',
@@ -574,6 +596,7 @@ class DevicenDisplay(DeviceUnreal):
                 'Local interface binding (-UDPMESSAGING_TRANSPORT_UNICAST) of '
                 'the form {address}:{port}. If {address} is omitted, the device '
                 'address is used.'),
+            category="Network Settings",
         ),
         'udpmessaging_extra_static_endpoints': StringSetting(
             attr_name='udpmessaging_extra_static_endpoints',
@@ -583,25 +606,29 @@ class DevicenDisplay(DeviceUnreal):
                 'Comma separated. Used to add static endpoints '
                 '(-UDPMESSAGING_TRANSPORT_STATIC) in addition to those '
                 'managed by Switchboard.'),
+            category="Network Settings",
         ),
         'disable_ensures': BoolSetting(
             attr_name='disable_ensures',
             nice_name="Disable Ensures",
             value=True,
-            tool_tip="When checked, disables the handling of ensure errors - which are non-fatal and may cause hitches."
+            tool_tip="When checked, disables the handling of ensure errors - which are non-fatal and may cause hitches.",
+            category="UE Settings",
         ),
         'disable_all_screen_messages': BoolSetting(
             attr_name='disable_all_screen_messages',
             nice_name="Disable All Screen Messages",
             value=True,
-            tool_tip="When checked, adds DisableAllScreenMessages to ExecCmds"
+            tool_tip="When checked, adds DisableAllScreenMessages to ExecCmds",
+            category="UE Settings",
         ),
         'livelink_preset': LiveLinkPresetSetting(
             attr_name='livelink_preset',
             nice_name='LiveLink Preset',
             value='',
             tool_tip=(
-                'Adds the selected LiveLink preset to the command line \n')
+                'Adds the selected LiveLink preset to the command line \n'),
+            category="Tools Settings",
         ),
         'graphics_adapter': OptionSetting(
             attr_name="graphics_adapter",
@@ -613,12 +640,14 @@ class DevicenDisplay(DeviceUnreal):
                 "- 'Config' : Use the setting in the nDisplay config file \n"
                 "- 0, 1, .. : The specified gpu index \n"
             ),
+            category="GPU/CPU Settings",
         ),
         'mediaprofile': MediaProfileSetting(
             attr_name='mediaprofile',
             nice_name='Media Profile',
             value='',
-            tool_tip=('Adds the selected Media Profile to the command line')
+            tool_tip=('Adds the selected Media Profile to the command line'),
+            category="Tools Settings",
         ),
         'lock_gpu_clock': BoolSetting(
             attr_name="lock_gpu_clock",
@@ -629,6 +658,7 @@ class DevicenDisplay(DeviceUnreal):
                 "to be running on the client machine as administrator, otherwise this option will be ignored."
             ),
             show_ui=True if sys.platform in ('win32', 'linux') else False,  # Gpu Clocker is available in select platforms
+            category="GPU/CPU Settings",
         ),
     }
 
@@ -644,7 +674,8 @@ class DevicenDisplay(DeviceUnreal):
                 nice_name="UE Command Line",
                 value=kwargs.get("ue_command_line", ''),
                 allow_reset=False,
-                is_read_only=True
+                is_read_only=True,
+                category="Command Line Args",
             ),
             'window_position': Setting(
                 attr_name="window_position",
diff --git a/Engine/Plugins/VirtualProduction/Switchboard/Source/Switchboard/switchboard/devices/unreal/plugin_unreal.py b/Engine/Plugins/VirtualProduction/Switchboard/Source/Switchboard/switchboard/devices/unreal/plugin_unreal.py
--- a/Engine/Plugins/VirtualProduction/Switchboard/Source/Switchboard/switchboard/devices/unreal/plugin_unreal.py
+++ b/Engine/Plugins/VirtualProduction/Switchboard/Source/Switchboard/switchboard/devices/unreal/plugin_unreal.py
@@ -18,7 +18,7 @@ import re
 import socket
 import sys
 import threading
-from typing import Callable, Generator, Optional, Union
+from typing import Callable, Dict, Generator, Optional, Union
 import uuid
 import time
 
@@ -330,25 +330,6 @@ class LiveLinkPresetSetting(Setting):
     the options.
     '''
 
-    def __init__(
-        self,
-        attr_name,
-        nice_name,
-        value,
-        tool_tip=None,
-        show_ui=True,
-        allow_reset=True,
-        migrate_data=None
-    ):
-        super().__init__(
-            attr_name=attr_name,
-            nice_name=nice_name,
-            value=value,
-            tool_tip=tool_tip,
-            show_ui=show_ui,
-            allow_reset=allow_reset,
-            migrate_data=migrate_data)
-
     def _create_widgets(self, override_device_name=None):
 
         # create combo with livelink preset options
@@ -521,32 +502,37 @@ class DeviceUnreal(Device):
             value=1024,
             tool_tip=(
                 "Buffer size used for communication with SwitchboardListener"),
+            category="UE Settings",
         ),
         'command_line_arguments': StringSetting(
             attr_name="command_line_arguments",
             nice_name='Command Line Arguments',
             value="",
             tool_tip='Additional command line arguments for the engine',
+            category="Command Line Args",
         ),
         'exec_cmds': StringListSetting(
             attr_name="exec_cmds",
             nice_name='ExecCmds',
             value=[],
             tool_tip='ExecCmds to be passed. No need for outer double quotes.',
-            migrate_data=migrate_comma_separated_string_to_list
+            migrate_data=migrate_comma_separated_string_to_list,
+            category="Command Line Args",
         ),
         'dp_cvars': StringListSetting(
             attr_name='dp_cvars',
             nice_name="DPCVars",
             value=[],
             tool_tip="Device profile console variables.",
-            migrate_data=migrate_comma_separated_string_to_list
+            migrate_data=migrate_comma_separated_string_to_list,
+            category="Command Line Args",
         ),
         'port': IntSetting(
             attr_name="port",
             nice_name="Listener Port",
             value=2980,
-            tool_tip="Port of SwitchboardListener"
+            tool_tip="Port of SwitchboardListener",
+            category="Network Settings",
         ),
         'osc_port': IntSetting(
             attr_name='osc_port',
@@ -555,6 +541,7 @@ class DeviceUnreal(Device):
             tool_tip=(
                 'Must match the port on which the Unreal Editor OSC server '
                 'is configured to listen for OSC connections.'),
+            category="Network Settings",
         ),
         'roles_filename': StringSetting(
             attr_name="roles_filename",
@@ -563,6 +550,7 @@ class DeviceUnreal(Device):
             tool_tip=(
                 "File that stores VirtualProduction roles. "
                 "Default: Config/Tags/VPRoles.ini"),
+            category="UE Settings",
         ),
         'stage_session_id': IntSetting(
             attr_name="stage_session_id",
@@ -572,11 +560,13 @@ class DeviceUnreal(Device):
                 "An ID that groups Stage Monitor providers and monitors. "
                 "Instances with different Session IDs are invisible to each "
                 "other in Stage Monitor."),
+            category="Tools Settings",
         ),
         'ue_exe': StringSetting(
             attr_name="editor_exe",
             nice_name="Unreal Editor filename",
             value="UnrealEditor.exe",
+            category="General Settings",
         ),
         'max_gpu_count': OptionSetting(
             attr_name="max_gpu_count",
@@ -586,6 +576,7 @@ class DeviceUnreal(Device):
             tool_tip=(
                 "If you have multiple GPUs in the PC, you can specify how "
                 "many to use."),
+            category="GPU/CPU Settings",
         ),
         'priority_modifier': OptionSetting(
             attr_name='priority_modifier',
@@ -593,6 +584,7 @@ class DeviceUnreal(Device):
             value=sb_utils.PriorityModifier.Normal.name,
             possible_values=[p.name for p in sb_utils.PriorityModifier],
             tool_tip="Used to override the priority of the process.",
+            category="GPU/CPU Settings",
         ),
         'auto_decline_package_recovery': BoolSetting(
             attr_name='auto_decline_package_recovery',
@@ -603,6 +595,7 @@ class DeviceUnreal(Device):
                 'skipping the restore prompt. Useful in multi-user '
                 'scenarios, where restoring from auto-save may be '
                 'undesirable.'),
+            category="UE Settings",
         ),
         'udpmessaging_unicast_endpoint': StringSetting(
             attr_name='udpmessaging_unicast_endpoint',
@@ -612,6 +605,7 @@ class DeviceUnreal(Device):
                 'Local interface binding (-UDPMESSAGING_TRANSPORT_UNICAST) of '
                 'the form {address}:{port}. If {address} is omitted, the '
                 'device address is used.'),
+            category="Network Settings",
         ),
         'udpmessaging_extra_static_endpoints': StringSetting(
             attr_name='udpmessaging_extra_static_endpoints',
@@ -621,6 +615,7 @@ class DeviceUnreal(Device):
                 'Comma separated. Used to add static endpoints '
                 '(-UDPMESSAGING_TRANSPORT_STATIC) in addition to those '
                 'managed by Switchboard.'),
+            category="Network Settings",
         ),
         'udpmessaging_multicast_endpoint': StringSetting(
             attr_name='udpmessaging_multicast_endpoint',
@@ -630,6 +625,7 @@ class DeviceUnreal(Device):
                 'Multicast group and port (-UDPMESSAGING_TRANSPORT_MULTICAST) '
                 'in the {address}:{port} endpoint format. The multicast group address '
                 'must be in the range 224.0.0.0 to 239.255.255.255.'),
+            category="Network Settings",
         ),
         'log_download_dir': DirectoryPathSetting(
             attr_name='log_download_dir',
@@ -638,6 +634,7 @@ class DeviceUnreal(Device):
             tool_tip=(
                 'Directory in which to store logs transferred from devices. '
                 'If unset, defaults to $(ProjectDir)/Saved/Logs/Switchboard/'),
+            category="General Settings",
         ),
         'reflect_visibility_to_game': BoolSetting(
             attr_name='reflect_visibility_to_game',
@@ -647,13 +644,15 @@ class DeviceUnreal(Device):
                 'Sets the value for `Reflect Level Visibilty to Game` for Multi-user editing. \n'
                 'Editor visibilty state will be applied to the game equivalent visibility properties. \n'
                 'This is useful for ICVFX workflows where the editor visibilty state directly \n'
-                'corresponds to the state on the render nodes.')
+                'corresponds to the state on the render nodes.'),
+            category="Multi-User Settings",
         ),
         'rsync_port': IntSetting(
             attr_name='rsync_port',
             nice_name='Rsync Server Port',
             value=switchboard_application.RsyncServer.DEFAULT_PORT,
-            tool_tip='Port number on which the rsync server should listen.'
+            tool_tip='Port number on which the rsync server should listen.',
+            category="Network Settings",
         ),
         'listener_inactive_timeout': IntSetting(
             attr_name='listener_inactive_timeout',
@@ -662,7 +661,8 @@ class DeviceUnreal(Device):
             tool_tip=(
                 'Tells the connected Listener to wait at least N seconds '
                 'between network messages before considering the connection '
-                'to Switchboard lost and closing it with a timeout error.')
+                'to Switchboard lost and closing it with a timeout error.'),
+            category="Network Settings",
         ),
         'slate_allow_throttling': BoolSetting(
             attr_name='slate_allow_throttling',
@@ -671,27 +671,31 @@ class DeviceUnreal(Device):
             tool_tip=(
                 'Sets the Slate.bAllowThrottling cvar. When unchecked, the Editor viewports do not freeze/throttle \n'
                 'during certain operations. Not thottling is typically desired when using the Editor in \n'
-                'a virtual production stage.\n')
+                'a virtual production stage.\n'),
+            category="UE Settings",
         ),
         'retrieve_logs': BoolSetting(
             attr_name='retrieve_logs',
             nice_name='Retrieve Logs',
             value=True,
             tool_tip=(
-                'When checked, retrieves the logs and traces after Unreal Engine terminates. \n')
+                'When checked, retrieves the logs and traces after Unreal Engine terminates. \n'),
+            category="UE Settings",
         ),
         'livelink_preset': LiveLinkPresetSetting(
             attr_name='livelink_preset',
             nice_name='LiveLink Preset',
             value='',
             tool_tip=(
-                'Adds the selected LiveLink preset to the command line \n')
+                'Adds the selected LiveLink preset to the command line \n'),
+            category="Tools Settings",
         ),
         'mediaprofile': MediaProfileSetting(
             attr_name='mediaprofile',
             nice_name='Media Profile',
             value='',
-            tool_tip=('Adds the selected Media Profile to the command line')
+            tool_tip=('Adds the selected Media Profile to the command line'),
+            category="Tools Settings",
         ),
         'lock_gpu_clock': BoolSetting(
             attr_name="lock_gpu_clock",
@@ -701,7 +705,8 @@ class DeviceUnreal(Device):
                 "Hint to lock the GPU clock to its allowed maximum. Requires SwitchboardListenerHelper \n"
                 "to be running on the client machine, otherwise this option will be ignored."
             ),
-            show_ui = True if sys.platform in ('win32','linux') else False, # Gpu Clocker is available in select platforms
+            show_ui=True if sys.platform in ('win32', 'linux') else False,  # Only available in select platforms
+            category="GPU/CPU Settings",
         ),
         'use_sync_filters': BoolSetting(
             attr_name='use_sync_filters',
@@ -709,7 +714,8 @@ class DeviceUnreal(Device):
             value=False,
             tool_tip=(
                 'Controls whether UnrealGameSync filter categories or custom '
-                'views specified below are used during Perforce sync')
+                'views specified below are used during Perforce sync'),
+            category="Source Control Settings",
         ),
         'included_sync_categories': MultiOptionSetting(
             attr_name='included_sync_categories',
@@ -725,6 +731,7 @@ class DeviceUnreal(Device):
                                    'Source Code'),
             ],
             tool_tip='UnrealGameSync filter categories to include during sync',
+            category="Source Control Settings",
         ).with_get_json_override_fn(
             lambda val_list: [(str(x.id), x.name) for x in val_list]
         ).with_config_set_override_fn(
@@ -740,6 +747,7 @@ class DeviceUnreal(Device):
             tool_tip=(
                 'Comma separated. Used to specify freeform Perforce-style '
                 'wildcards to be applied during sync.'),
+            category="Source Control Settings",
         ),
     }
 
@@ -751,7 +759,8 @@ class DeviceUnreal(Device):
             tool_tip=(
                 "Expects an absolute path to the device directory containing the ugs.dll library file. \n"
                 "If left blank, switchboard will attempt to find UGS by searching the system PATH and default install locations. \n"
-                "On Windows, the default install location is '${LOCALAPPDATA}/UnrealGameSync/Latest/'.")
+                "On Windows, the default install location is '${LOCALAPPDATA}/UnrealGameSync/Latest/'."),
+            category="UGS Settings",
         )
 
     # A stand-in QObject to emit signals from classmethods.
@@ -942,14 +951,16 @@ class DeviceUnreal(Device):
             attr_name="roles",
             nice_name="Roles",
             value=roles,
-            tool_tip="List of roles for this device"
+            tool_tip="List of roles for this device",
+            category="UE Settings",
         )
 
         self.setting_ddc_build_platforms = MultiOptionSetting(
             attr_name="ddc_build_platforms",
             nice_name="DDC Build Platforms",
             value=["Windows", "WindowsEditor", "Linux", "LinuxEditor"],
-            tool_tip="List of platforms for which to build the DDC cache for this device"
+            tool_tip="List of platforms for which to build the DDC cache for this device",
+            category="General Settings",
         )
 
         autojoin_mu_server = kwargs.get("autojoin_mu_server", True)
@@ -957,49 +968,51 @@ class DeviceUnreal(Device):
             attr_name="autojoin_mu_server",
             nice_name="Auto join Multi-user Server",
             value=autojoin_mu_server,
-            show_ui=False
+            show_ui=False,
         )
 
         self.last_launch_command = StringSetting(
             attr_name="last_launch_command",
             nice_name="Last Launch Command",
             value=kwargs.get("last_launch_command", ''),
-            show_ui=False
+            show_ui=False,
         )
 
         self.last_log_path = FilePathSetting(
             attr_name="last_log_path",
             nice_name="Last Log Path",
             value=kwargs.get("last_log_path", ''),
-            show_ui=False
+            show_ui=False,
         )
 
         self.last_trace_path = FilePathSetting(
             attr_name="last_trace_path",
             nice_name="Last Insights Trace Path",
             value=kwargs.get("last_trace_path", ''),
-            show_ui=False
+            show_ui=False,
         )
 
         self.last_sync_filter_hash = StringSetting(
             attr_name="last_sync_filter_hash",
             nice_name="Last UGS Sync Filter Hash",
             value=kwargs.get('last_sync_filter_hash', ''),
-            show_ui=False
+            show_ui=False,
         )
 
         self.exclude_from_build = BoolSetting(
             attr_name="exclude_from_build",
             nice_name="Exclude from build",
             value=kwargs.get("exclude_from_build", False),
-            tool_tip="Whether to exclude this device from builds"
+            tool_tip="Whether to exclude this device from builds",
+            category="General Settings",
         )
 
         self.exclude_from_insights = BoolSetting(
             attr_name="exclude_from_insights",
             nice_name="Exclude from Insights trace",
             value=kwargs.get("exclude_from_insights", False),
-            tool_tip="Whether to exclude device from Unreal Insights traces"
+            tool_tip="Whether to exclude device from Unreal Insights traces",
+            category="General Settings",
         )
 
         self.setting_address.signal_setting_changed.connect(
@@ -1234,6 +1247,48 @@ class DeviceUnreal(Device):
 
         return overrides
 
+    @classmethod
+    def sort_setting_categories_from_preferred_list(cls, categories: Dict, preferred_order: list[str]) -> OrderedDict:
+        ''' Sorts using the order in the given list of preferred categories '''
+
+        # Create a dictionary to hold the preferred categories in their given order
+        # Key will be the string, value will be the index  in the preferred_order list.
+        preferred_dict = {key: idx for (idx, key) in enumerate(preferred_order)}
+
+        # Custom sort function
+        def custom_sort_key(item):
+            key = item[0]
+            # Sort by the preferred order if it exists. If it is not in the preferred list,
+            # they will go at the end and sorted alphabetically
+            return (preferred_dict.get(key, len(preferred_order)), key)
+
+        # Sort the categories using the custom sort key. The sort keys returned for a given
+        # element will be compared against others by sorted.
+        return OrderedDict(sorted(categories.items(), key=custom_sort_key))
+
+    @classmethod
+    def sort_setting_categories(cls, categories: Dict) -> OrderedDict:
+        ''' Overrides base implementation '''
+
+        preferred_order = [
+            'General Settings',
+            'Source Control Settings',
+            'Command Line Args',
+            'Tools Settings',
+            'GPU/CPU Settings',
+            'UE Settings',
+            'Render Settings',
+            'Network Settings',
+            'Multi-User Settings',
+            'Multiplayer Settings',
+            'Device Settings',
+            'Misc',
+        ]
+
+        return DeviceUnreal.sort_setting_categories_from_preferred_list(
+            categories=categories,
+            preferred_order=preferred_order)
+
     def device_settings(self):
         return super().device_settings() + [
             self.setting_roles,
diff --git a/Engine/Plugins/VirtualProduction/Switchboard/Source/Switchboard/switchboard/settings_dialog.py b/Engine/Plugins/VirtualProduction/Switchboard/Source/Switchboard/switchboard/settings_dialog.py
--- a/Engine/Plugins/VirtualProduction/Switchboard/Source/Switchboard/switchboard/settings_dialog.py
+++ b/Engine/Plugins/VirtualProduction/Switchboard/Source/Switchboard/switchboard/settings_dialog.py
@@ -2,11 +2,10 @@
 
 import os
 import pathlib
-import typing
+from typing import Dict, Optional
 
 from PySide6 import QtCore
 from PySide6 import QtGui
-from PySide6 import QtUiTools
 from PySide6 import QtWidgets
 
 from switchboard import config
@@ -15,17 +14,20 @@ from switchboard.config import CONFIG, SETTINGS, Config, UserSettings, Setting
 from switchboard.settings_search import SettingsSearch
 from switchboard.switchboard_widgets import CollapsibleGroupBox, DropDownMenuComboBox
 from switchboard.ui.horizontal_tabs import HorizontalTabWidget
+from switchboard.devices.device_base import Device
 
 RELATIVE_PATH = os.path.dirname(__file__)
 
 
 def clear_widgets(layout):
     for i in range(layout.count()):
-        layout.takeAt(0) 
+        layout.takeAt(0)
+
 
 COLLAPSE_ALL_TEXT = "Collapse all"
 EXPAND_ALL_TEXT = "Expand all"
 
+
 class SettingsDialog(QtCore.QObject):
     def __init__(self, settings: UserSettings, config: Config):
         super().__init__()
@@ -33,7 +35,7 @@ class SettingsDialog(QtCore.QObject):
 
         # Set the UI object
         self.ui = QtWidgets.QDialog()
-        self.ui.resize(650, 800)
+        self.ui.resize(800, 800)
         dialog_layout = QtWidgets.QVBoxLayout(self.ui)
         dialog_layout.setContentsMargins(2, 2, 2, 2)
         self.ui.setWindowTitle("Settings")
@@ -49,37 +51,37 @@ class SettingsDialog(QtCore.QObject):
             self._create_project_settings(config),
             self._create_multi_user_server_settings(config)
         ]
-        
+
         self._create_tab_widget(dialog_layout)
-        
+
         self.settings_search = SettingsSearch([self.ui.all_settings_scroll_area, self.ui.general_settings_scroll_area])
         self.ui.searchBar.textChanged.connect(
             self._on_search_text_edited)
-        
+
     def select_all_tab(self):
         self._on_tab_changed(0)
-        
+
     def _on_tab_changed(self, index: int):
         """
         Because of the All category, widgets need to be re-parented constantly.
         """
         clear_widgets(self.ui.all_settings_scroll_area.layout())
         clear_widgets(self.ui.general_settings_scroll_area.layout())
-        
+
         is_all_category = index == 0
         if is_all_category:
             for general_setting_widget in self.general_settings_list:
                 self.ui.all_settings_scroll_area.layout().addWidget(general_setting_widget)
-               
+
             for (group_box, scroll_bar) in self.plugin_widgets.values():
                 scroll_bar.takeWidget()
                 self.ui.all_settings_scroll_area.layout().addWidget(group_box)
-                
+
             # Pull all widgets up so they do not attempt to divide space among themselves when filtered by search
             self.ui.all_settings_scroll_area.layout().addItem(
                 QtWidgets.QSpacerItem(20, 40, QtWidgets.QSizePolicy.Minimum, QtWidgets.QSizePolicy.Expanding)
             )
-        
+
         # General settings may have been in All
         is_general_category = index == 1
         if is_general_category:
@@ -89,22 +91,22 @@ class SettingsDialog(QtCore.QObject):
             self.ui.general_settings_scroll_area.layout().addItem(
                 QtWidgets.QSpacerItem(20, 40, QtWidgets.QSizePolicy.Minimum, QtWidgets.QSizePolicy.Expanding)
             )
-           
+
         # Add the plugin back to its tabs - it may have been in All
-        is_plugin = not is_all_category and not is_general_category 
+        is_plugin = not is_all_category and not is_general_category
         if is_plugin:
             plugin_name = self.ui.tab_widget.tabText(index)
             (group_box, scroll_bar) = self.plugin_widgets[plugin_name]
             is_plugin_widget_missing = scroll_bar.widget() is None
             if is_plugin_widget_missing:
                 scroll_bar.setWidget(group_box)
-        
+
         self._search_settings_again()
         return
 
     def _on_finished(self, result: int, config: Config):
         config.P4_ENABLED.signal_setting_changed.disconnect(self._on_source_control_setting_changed)
-        
+
         # Currently, the only way to dismiss the settings dialog is by using
         # the close button as opposed to ok/cancel buttons, so we intercept the
         # close to issue a warning if the config path was changed and we're
@@ -130,12 +132,12 @@ class SettingsDialog(QtCore.QObject):
 
     def _on_search_text_edited(self, search_string: str):
         self.settings_search.search(search_string)
-        
+
     def _create_search_area(self, dialog_layout):
         self.ui.search_area = QtWidgets.QWidget()
         search_layout = QtWidgets.QHBoxLayout(self.ui.search_area)
         search_layout.setContentsMargins(1, 1, 1, 1)
-        
+
         self.ui.searchBar = QtWidgets.QLineEdit()
         self.ui.searchBar.setPlaceholderText("Search")
         search_layout.addWidget(self.ui.searchBar)
@@ -144,7 +146,7 @@ class SettingsDialog(QtCore.QObject):
         self.ui.view_options = DropDownMenuComboBox(QtGui.QIcon(pixmap))
         self.ui.view_options.addItem(COLLAPSE_ALL_TEXT)
         self.ui.view_options.addItem(EXPAND_ALL_TEXT)
-        
+
         self.ui.view_options.on_select_option.connect(self._on_view_option_selected)
         search_layout.addWidget(self.ui.view_options)
         dialog_layout.addWidget(self.ui.search_area)
@@ -154,39 +156,39 @@ class SettingsDialog(QtCore.QObject):
             self._set_categories_expanded(False)
         elif selected_item == EXPAND_ALL_TEXT:
             self._set_categories_expanded(True)
-            
+
     def _set_categories_expanded(self, should_expand: bool):
         def _set_expanded_recursive(widget_or_layout, should_expand: bool):
             if isinstance(widget_or_layout, CollapsibleGroupBox):
                 widget_or_layout.set_expanded(should_expand)
-            
+
             layout = widget_or_layout
             if isinstance(widget_or_layout, QtWidgets.QWidget):
                 if isinstance(widget_or_layout, QtWidgets.QScrollArea):
                     layout = widget_or_layout.widget().layout()
                 else:
                     layout = widget_or_layout.layout()
-                
+
             if layout is None:
                 return
-            
+
             for i in range(layout.count()):
                 layout_item = layout.itemAt(i)
                 if layout_item.widget():
                     _set_expanded_recursive(layout_item.widget(), should_expand)
                 elif layout_item.layout():
                     _set_expanded_recursive(layout_item.layout(), should_expand)
-            
+
         active_tab_content = self.ui.tab_widget.currentWidget()
         _set_expanded_recursive(active_tab_content, should_expand)
-        
+
     def _create_config_path_settings(self):
         self.ui.config_path_layout = QtWidgets.QWidget()
         layout = QtWidgets.QHBoxLayout(self.ui.config_path_layout)
-        
+
         self.ui.config_path_label = QtWidgets.QLabel()
         self.ui.config_path_label.setText("Config Path")
-        
+
         self.ui.config_path_line_edit = QtWidgets.QLineEdit()
         self.ui.config_path_line_edit.textChanged.connect(
             self.config_path_text_changed)
@@ -197,9 +199,9 @@ class SettingsDialog(QtCore.QObject):
         # Store the current config paths so we can warn about overwriting an existing config.
         self._config_paths = config.list_config_paths()
         self.set_config_path(SETTINGS.CONFIG)
-        
+
         return self.ui.config_path_layout
-        
+
     def _create_switchboard_settings(self, settings: UserSettings, config: Config):
         self.ui.switchboard_settings_group = CollapsibleGroupBox()
         self.ui.switchboard_settings_group.setTitle("Switchboard")
@@ -211,30 +213,30 @@ class SettingsDialog(QtCore.QObject):
         config.SBLHELPER_EXE.create_ui(form_layout=layout)
 
         return self.ui.switchboard_settings_group
-        
+
     def _create_project_settings(self, config: Config):
         self.ui.project_settings_group = CollapsibleGroupBox()
         self.ui.project_settings_group.setTitle("Project Settings")
         layout = QtWidgets.QVBoxLayout(self.ui.project_settings_group)
-        
+
         project_settings_root = QtWidgets.QWidget()
         form_layout = QtWidgets.QFormLayout(project_settings_root)
         self._create_settings_section(config.basic_project_settings, form_layout)
         layout.addWidget(project_settings_root)
-        
+
         # Sub settings
         self._create_osc_settings(config, layout)
         self._create_source_control_settings(config, layout)
-        
+
         return self.ui.project_settings_group
-    
+
     def _create_osc_settings(self, config: Config, layout: QtWidgets.QFormLayout):
         self.ui.osc_settings_group = QtWidgets.QGroupBox()
         self.ui.osc_settings_group.setTitle("OSC")
         form_layout = QtWidgets.QFormLayout(self.ui.osc_settings_group)
         layout.addWidget(self.ui.osc_settings_group)
         self._create_settings_section(config.osc_settings, form_layout)
-    
+
     def _create_source_control_settings(self, config: Config, layout: QtWidgets.QFormLayout):
         self.ui.source_control_settings_group = QtWidgets.QGroupBox()
         self.ui.source_control_settings_group.setTitle("Source Control")
@@ -258,7 +260,7 @@ class SettingsDialog(QtCore.QObject):
 
     def _on_group_source_control_changed(self):
         CONFIG.P4_ENABLED.update_value(self.ui.source_control_settings_group.isChecked())
-        CONFIG.save();
+        CONFIG.save()
 
     def _on_source_control_setting_changed(self, old_value: bool, new_value: bool):
         self.ui.source_control_settings_group.setChecked(new_value)
@@ -267,11 +269,11 @@ class SettingsDialog(QtCore.QObject):
         self.ui.multi_user_settings = CollapsibleGroupBox()
         self.ui.multi_user_settings.setTitle("Multi User Server")
         form_layout = QtWidgets.QFormLayout(self.ui.multi_user_settings)
-        
+
         self._create_settings_section(config.mu_settings, form_layout)
         return self.ui.multi_user_settings
-    
-    def _create_settings_section(self, dict: typing.Dict[str, Setting], layout: QtWidgets.QFormLayout):
+
+    def _create_settings_section(self, dict: Dict[str, Setting], layout: QtWidgets.QFormLayout):
         setting_list = [setting for _, setting in dict.items()]
         for setting in setting_list:
             setting.create_ui(override_device_name=None, form_layout=layout)
@@ -280,7 +282,7 @@ class SettingsDialog(QtCore.QObject):
         self.ui.tab_widget = HorizontalTabWidget()
         self.ui.tab_widget.setTabPosition(QtWidgets.QTabWidget.West)
         parent_layout.addWidget(self.ui.tab_widget)
-    
+
         # All tab
         self.ui.all_tab_root = QtWidgets.QScrollArea()
         self.ui.all_tab_root.setWidgetResizable(True)
@@ -291,13 +293,13 @@ class SettingsDialog(QtCore.QObject):
         all_layout = QtWidgets.QVBoxLayout(self.ui.all_settings_scroll_area)
         all_layout.setContentsMargins(2, 2, 2, 2)
         self.ui.all_tab_root.setWidget(self.ui.all_settings_scroll_area)
-    
+
         # General tab
         self.ui.general_tab_root = QtWidgets.QScrollArea()
         self.ui.general_tab_root.setWidgetResizable(True)
         self.ui.general_tab_root.setSizeAdjustPolicy(QtWidgets.QAbstractScrollArea.AdjustToContents)
         self.ui.general_tab_root.setSizePolicy(QtWidgets.QSizePolicy.Expanding, QtWidgets.QSizePolicy.Expanding)
-        
+
         self.ui.general_settings_scroll_area = QtWidgets.QWidget()
         general_layout = QtWidgets.QVBoxLayout(self.ui.general_settings_scroll_area)
         general_layout.setContentsMargins(2, 2, 2, 2)
@@ -351,62 +353,131 @@ class SettingsDialog(QtCore.QObject):
 
         self._changed_config_path = config_path
 
+    def add_plugin_settings_widgets_to_layout(
+            self,
+            plugin_cls: Device,
+            plugin_settings,
+            layout: QtWidgets.QLayout,
+            stylevariant: str = '',
+            device_name: Optional[str] = None):
+        ''' Organizes plugin_settings into categories group boxes and adds them to the given layout'''
+
+        scategories = {}
+
+        for setting in plugin_settings:
+
+            # Don't create categories for hidden properties, or you might end up with an empty category
+            if not setting.show_ui:
+                continue
+
+            category = setting.category
+
+            if category not in scategories:
+                scategories[category] = []
+
+            scategories[category].append(setting)
+
+        # sort the categories per the preference of the plugin
+        scategories = plugin_cls.sort_setting_categories(scategories)
+
+        # Create a group box for each category and add it to the plugin layout
+        for category, settings in scategories.items():
+
+            scategory_group_box = QtWidgets.QGroupBox()
+            scategory_group_box.setTitle(category)
+            scategory_group_box.setLayout(QtWidgets.QVBoxLayout())
+            scategory_group_box.setSizePolicy(
+                QtWidgets.QSizePolicy(QtWidgets.QSizePolicy.Expanding, QtWidgets.QSizePolicy.Maximum)
+            )
+
+            # Alter the style as necessary to discern group boxes inside group boxes
+            scategory_group_box.setProperty('stylevariant', stylevariant)
+
+            scategory_layout = QtWidgets.QFormLayout()
+            scategory_group_box.layout().addLayout(scategory_layout)
+
+            # Add widgets for each setting in this category
+            for setting in settings:
+                setting.create_ui(form_layout=scategory_layout, override_device_name=device_name)
+
+            # Add the category group box to the main plugin layout
+            layout.addWidget(scategory_group_box)
+
     # Devices
-    def add_section_for_plugin(self, plugin_name, plugin_settings, device_settings):
+    def add_section_for_plugin(
+        self,
+        plugin_name: str,
+        plugin_cls: Device,
+        plugin_settings,
+        device_settings
+    ):
+
         any_device_settings = (
             any([device[1] for device in device_settings]) or
             any([device[2] for device in device_settings]))
+
         if not any_device_settings:
             return  # no settings to show
 
         # Create a group box per plugin
-        device_override_group_box = CollapsibleGroupBox()
-        device_override_group_box.setTitle(f'{plugin_name} Settings')
-        device_override_group_box.setLayout(QtWidgets.QVBoxLayout())
-        device_override_group_box.setSizePolicy(
+        plugin_group_box = CollapsibleGroupBox()
+        plugin_group_box.setTitle(f'{plugin_name} Settings')
+        plugin_group_box.setLayout(QtWidgets.QVBoxLayout())
+        plugin_group_box.setSizePolicy(
             QtWidgets.QSizePolicy(QtWidgets.QSizePolicy.Expanding, QtWidgets.QSizePolicy.Maximum)
         )
 
         # Need to manually create scroll bar for tab
         scroll_bar = QtWidgets.QScrollArea()
-        scroll_bar.setWidget(device_override_group_box)
+        scroll_bar.setWidget(plugin_group_box)
         scroll_bar.setWidgetResizable(True)
         scroll_bar.setSizeAdjustPolicy(QtWidgets.QAbstractScrollArea.AdjustIgnored)
         scroll_bar.setSizePolicy(QtWidgets.QSizePolicy.Expanding, QtWidgets.QSizePolicy.Expanding)
         scroll_layout = QtWidgets.QVBoxLayout(scroll_bar)
         scroll_layout.setContentsMargins(1, 1, 1, 1)
-        
+
         self.ui.tab_widget.addTab(scroll_bar, f'{plugin_name}')
-        self.plugin_widgets[plugin_name] = (device_override_group_box, scroll_bar)
-        self.settings_search.searched_widgets.append(device_override_group_box)
+        self.plugin_widgets[plugin_name] = (plugin_group_box, scroll_bar)
 
-        plugin_layout = QtWidgets.QFormLayout()
-        device_override_group_box.layout().addLayout(plugin_layout)
+        # Add the device widget (contains all widgets and layouts for the given plugin)
+        # to the list of widgets to be searched for matches. There should not be a need
+        # to add child widgets since that would cause double searches.
+        self.settings_search.searched_widgets.append(plugin_group_box)
 
-        # add widgets for settings shared by all devices of a plugin
-        for setting in plugin_settings:
-            setting.create_ui(form_layout=plugin_layout)
+        # Add a layout to the plugin group box. It will contain the plugin settings and the
+        # device intances and their settings.
+        plugin_group_box.layout().addLayout(QtWidgets.QVBoxLayout())
+
+        self.add_plugin_settings_widgets_to_layout(
+            plugin_cls=plugin_cls,
+            plugin_settings=plugin_settings,
+            layout=plugin_group_box.layout())
 
         # add widgets for settings and overrides of individual devices
         for device_name, settings, overrides in device_settings:
+
+            # Each device instance will have its own collapsible group box
             group_box = CollapsibleGroupBox()
             group_box.setTitle(device_name)
             group_box.setLayout(QtWidgets.QVBoxLayout())
 
-            device_override_group_box.layout().addWidget(group_box)
+            plugin_group_box.layout().addWidget(group_box)
 
-            layout = QtWidgets.QFormLayout()
+            layout = QtWidgets.QVBoxLayout()
             group_box.layout().addLayout(layout)
 
-            # regular "instance" settings
-            for setting in settings:
-                setting.create_ui(form_layout=layout)
+            # Merge instance settings and class overrides to avoid repeating
+            # category group boxes under the same device instance.
+            instance_settings = settings + overrides
 
-            for setting in overrides:
-                setting.create_ui(
-                    form_layout=layout, override_device_name=device_name)
+            self.add_plugin_settings_widgets_to_layout(
+                plugin_cls=plugin_cls,
+                plugin_settings=instance_settings,
+                layout=layout,
+                stylevariant='device_override',
+                device_name=device_name)
 
         # Pull all widgets up so they do not attempt to divide space among themselves when filtered by search
-        device_override_group_box.layout().addItem(
+        plugin_group_box.layout().addItem(
             QtWidgets.QSpacerItem(0, 0, QtWidgets.QSizePolicy.Minimum, QtWidgets.QSizePolicy.Expanding)
         )
diff --git a/Engine/Plugins/VirtualProduction/Switchboard/Source/Switchboard/switchboard/settings_search.py b/Engine/Plugins/VirtualProduction/Switchboard/Source/Switchboard/switchboard/settings_search.py
--- a/Engine/Plugins/VirtualProduction/Switchboard/Source/Switchboard/switchboard/settings_search.py
+++ b/Engine/Plugins/VirtualProduction/Switchboard/Source/Switchboard/switchboard/settings_search.py
@@ -1,14 +1,7 @@
 # Copyright Epic Games, Inc. All Rights Reserved.
 
-import os
-import pathlib
 from typing import List, Dict
-
-from PySide6 import QtCore
-from PySide6 import QtGui
-from PySide6 import QtUiTools
 from PySide6 import QtWidgets
-
 from switchboard.switchboard_widgets import CollapsibleGroupBox
 
 
@@ -19,12 +12,12 @@ def as_widget_or_layout(row_item):
 class FormLayoutRowVisibility:
     """
     Handles the showing and hiding of rows in QFormLayout
-    
+
     Row widgets are removed when hiding and added back when showing instead of calling setVisible on the row's widgets:
     Even if a row only contains invisible widgets, QFormLayout.layoutVerticalSpacing still affects them and visually 
     shifts visible rows around.
     """
-    
+
     def __init__(self, form: QtWidgets.QFormLayout):
         self.form = form
         # Save the original widget order so we can preserve it when adding hidden widgets back
@@ -39,20 +32,20 @@ class FormLayoutRowVisibility:
                     True
                 )
                 self.rows.append(pair)
-        
+
     def row_count(self):
         return len(self.rows)
-        
+
     def get_label_item(self, index: int):
         return self.rows[index][0]
-        
+
     def set_row_visible(self, index: int, visible: bool):
         (label, field, is_visible) = self.rows[index]
         if visible and not is_visible:
             self._show_row(index)
         elif not visible and is_visible:
             self._hide_row(index)
-                    
+
     def _show_row(self, index: int):
         (label, field, is_visible) = self.rows[index]
         index_to_use = 0 if index == 0 else self._find_index_to_insert_row(index)
@@ -62,11 +55,11 @@ class FormLayoutRowVisibility:
             field,
             True
         )
-        
+
         # Widgets were hidden when the row was hidden
         self._set_visible_recursive(label, True)
         self._set_visible_recursive(field, True)
-        
+
     def _hide_row(self, index: int):
         (label, field, is_visible) = self.rows[index]
         # "index" is the original index of the row before we any rows were hidden.
@@ -82,16 +75,16 @@ class FormLayoutRowVisibility:
     def _take_row(self, index: int):
         label_item = self.form.itemAt(index, QtWidgets.QFormLayout.LabelRole)
         field_item = self.form.itemAt(index, QtWidgets.QFormLayout.FieldRole)
-        
+
         # Otherwise there will be artifacts in the UI texture
         self._set_visible_recursive(as_widget_or_layout(label_item), False)
         self._set_visible_recursive(as_widget_or_layout(field_item), False)
-        
+
         # Remove items before removeRow so the widgets are not destroyed
         self.form.removeItem(label_item)
         self.form.removeItem(field_item)
         self.form.removeRow(index)
-        
+
         # Reallocate space
         self.form.invalidate()
         self.form.update()
@@ -105,12 +98,12 @@ class FormLayoutRowVisibility:
                     return actual_form_index + 1
 
         return 0
-            
+
     def _set_visible_recursive(self, element, visible: bool):
         if isinstance(element, QtWidgets.QWidget):
             element.setVisible(visible)
             return
-        
+
         for i in range(element.count()):
             child_item = element.itemAt(i)
             if child_item.widget():
@@ -122,79 +115,120 @@ class FormLayoutRowVisibility:
 class SettingsSearch:
     def __init__(self, searched_widgets: List[QtWidgets.QWidget]):
         self.searched_widgets = searched_widgets
-        
+
         form_layout_handlers: Dict[QtWidgets.QFormLayout, FormLayoutRowVisibility] = {}
         self.form_layout_handlers = form_layout_handlers
-        
+
     def search(self, search_string: str):
         search_term_list = search_string.split()
         for widget in self.searched_widgets:
             self._update_widget_visibility(widget, search_term_list)
 
     def _update_widget_visibility(self, widget: QtWidgets.QWidget, search_term_list):
-        # We only search labels for search terms
+
+        # We only search labels for search terms.
         if isinstance(widget, QtWidgets.QLabel):
             matches_search = self._is_search_match(widget.text(), search_term_list)
             widget.setVisible(matches_search)
             return matches_search
-        
+
         if widget.layout():
-            is_category = isinstance(widget, QtWidgets.QGroupBox)
-            is_match_on_containing_category = is_category and self._is_search_match(widget.title(), search_term_list)
-            # Make all children appear
-            if is_match_on_containing_category:
+            is_group_box = isinstance(widget, QtWidgets.QGroupBox)
+            is_match_on_group_box = is_group_box and self._is_search_match(widget.title(), search_term_list)
+
+            # If group box title matched, force all children to visible by clearing the search terms.
+            if is_match_on_group_box:
                 search_term_list = []
 
-            is_any_child_visible = self._update_layout_visibility(widget.layout(), search_term_list) \
-                or is_match_on_containing_category
-            widget.setVisible(is_any_child_visible)
-            if is_any_child_visible and isinstance(widget, CollapsibleGroupBox):
+            # Recursively search inside the layout of this widget
+            is_any_child_visible = self._update_layout_visibility(widget.layout(), search_term_list)
+
+            # Widget should be visible if its title matches or any of its children is visible
+            widget_visibility = is_match_on_group_box or is_any_child_visible
+            widget.setVisible(widget_visibility)
+
+            # Auto-expand visible collapsible group boxes
+            if widget_visibility and isinstance(widget, CollapsibleGroupBox):
                 widget.set_expanded(True)
-            
-            return is_any_child_visible
+
+            return widget_visibility
 
         return False
 
     def _update_layout_visibility(self, layout: QtWidgets.QLayout, search_term_list):
+
         is_match_on_any_child = False
 
+        # Form layouts are treated especially: It peruses its rows for QLabels that match the search terms,
+        # and doesn't look any further.
         if isinstance(layout, QtWidgets.QFormLayout):
             return self._handle_form_layout(layout, search_term_list)
 
-        # Is this horizontal boxes similar to a form?
+        # We assume that an HBox layout with 2 items is a setting (key_label | value_widget).
         if isinstance(layout, QtWidgets.QHBoxLayout) and layout.count() == 2:
             first = layout.itemAt(0)
             first_item_is_label = first and first.widget() and isinstance(first.widget(), QtWidgets.QLabel)
             if first_item_is_label:
-                is_match_on_any_child = self._update_widget_visibility(
-                    first.widget(), search_term_list)
+                is_match_on_any_child = self._update_widget_visibility(first.widget(), search_term_list)
                 return is_match_on_any_child
 
-        for i in range(layout.count()):
-            layout_item = layout.itemAt(i)
+        # Recursively iterate over all child widgets and layouts, and update is_match_on_any_child accordingly.
+        for idx in range(layout.count()):
+            layout_item = layout.itemAt(idx)
             if layout_item.widget():
                 is_match_on_any_child |= self._update_widget_visibility(layout_item.widget(), search_term_list)
-            elif layout_item.layout():
+            if layout_item.layout():
                 is_match_on_any_child |= self._update_layout_visibility(layout_item.layout(), search_term_list)
 
         return is_match_on_any_child
-    
+
     def _handle_form_layout(self, layout: QtWidgets.QFormLayout, search_term_list):
+        '''
+        Handles the visibility of rows in a QFormLayout based on a list of search terms.
+
+        This method iterates through each row in the provided QFormLayout, checking if the label
+        in each row matches any of the provided search terms. If a match is found, the corresponding
+        row is made visible; otherwise, it is hidden. The method returns a boolean indicating whether
+        any rows matched the search criteria.
+
+        Args:
+            layout (QtWidgets.QFormLayout): The form layout whose rows are to be checked and updated.
+            search_term_list (List[str]): A list of search terms used to determine visibility.
+
+        Returns:
+            bool: True if any row in the layout matches the search terms, False otherwise.
+        '''
+
         is_match_on_any_child = False
         form_handler = self.form_layout_handlers.setdefault(layout, FormLayoutRowVisibility(layout))
-        for i in range(form_handler.row_count()):
-            label = form_handler.get_label_item(i)
+
+        for idx in range(form_handler.row_count()):
+            label = form_handler.get_label_item(idx)
             if isinstance(label, QtWidgets.QWidget):
                 is_match = self._update_widget_visibility(label, search_term_list)
                 is_match_on_any_child |= is_match
-                form_handler.set_row_visible(i, is_match)
+                form_handler.set_row_visible(idx, is_match)
 
         return is_match_on_any_child
 
     @staticmethod
-    def _is_search_match(text: str, search_term_list: str):
+    def _is_search_match(text: str, search_term_list: List[str]):
+        '''
+        Determines whether the given text matches all search terms in the search term list.
+
+        This method checks if each search term in the provided list appears within the
+        given text. Both the text and the search terms are compared in a case-insensitive
+        manner. The method returns True only if all search terms are found in the text,
+        or search_term_list is empty.
+
+        Args:
+            text (str): The text to search within.
+            search_term_list (List[str]): A list of search terms to match against the text.
+
+        Returns:
+            bool: True if all search terms are found in text or if search_term_list is empty.
+        '''
         for search_term in search_term_list:
             if not search_term.lower() in text.lower():
                 return False
         return True
-    
\ No newline at end of file
diff --git a/Engine/Plugins/VirtualProduction/Switchboard/Source/Switchboard/switchboard/switchboard_dialog.py b/Engine/Plugins/VirtualProduction/Switchboard/Source/Switchboard/switchboard/switchboard_dialog.py
--- a/Engine/Plugins/VirtualProduction/Switchboard/Source/Switchboard/switchboard/switchboard_dialog.py
+++ b/Engine/Plugins/VirtualProduction/Switchboard/Source/Switchboard/switchboard/switchboard_dialog.py
@@ -856,7 +856,7 @@ class SwitchboardDialog(QtCore.QObject):
 
     def show_device_add_menu(self):
         self.device_add_menu.clear()
-        plugins = sorted(self.device_manager.available_device_plugins(), key=str.lower)
+        plugins = sorted(self.device_manager.available_device_plugins().keys(), key=str.lower)
         for plugin in plugins:
             icons = self.device_manager.plugin_icons(plugin)
             icon = icons["enabled"] if "enabled" in icons.keys() else QtGui.QIcon()
@@ -1211,10 +1211,21 @@ class SwitchboardDialog(QtCore.QObject):
         # TODO: VALIDATE RECORD PATH
         settings_dialog = SettingsDialog(SETTINGS, CONFIG)
 
-        for plugin_name in sorted(self.device_manager.available_device_plugins(), key=str.lower):
+        for plugin_name in sorted(self.device_manager.available_device_plugins().keys(), key=str.lower):
+
             device_instances = self.device_manager.devices_of_type(plugin_name)
-            device_settings = [(device.name, device.device_settings(), device.setting_overrides()) for device in device_instances]
-            settings_dialog.add_section_for_plugin(plugin_name, self.device_manager.plugin_settings(plugin_name), device_settings)
+
+            device_settings = [(device.name, device.device_settings(), device.setting_overrides())
+                               for device in device_instances]
+
+            plugin_cls = self.device_manager.available_device_plugins()[plugin_name]
+
+            settings_dialog.add_section_for_plugin(
+                plugin_name=plugin_name,
+                plugin_cls=plugin_cls,
+                plugin_settings=self.device_manager.plugin_settings(plugin_name),
+                device_settings=device_settings
+            )
 
         settings_dialog.select_all_tab()
 
