SB: Fixed project name when creating a new config.

It was using the name of the current config.

#rnx

diff --git a/Engine/Plugins/VirtualProduction/Switchboard/Source/Switchboard/switchboard/config.py b/Engine/Plugins/VirtualProduction/Switchboard/Source/Switchboard/switchboard/config.py
--- a/Engine/Plugins/VirtualProduction/Switchboard/Source/Switchboard/switchboard/config.py
+++ b/Engine/Plugins/VirtualProduction/Switchboard/Source/Switchboard/switchboard/config.py
@@ -2117,13 +2117,15 @@ class Config(object):
     def init_new_config(self, file_path: Union[str, pathlib.Path], uproject, engine_dir, p4_settings):
         ''' Initialize new configuration '''
 
+        # Assign self.file_path before the basic_project_settings line because it requires it to be updated.
+        self.file_path = get_absolute_config_path(file_path)
+
         basic_project_settings = {
                 "project_name": self.file_path.stem,
                 "uproject": uproject,
                 "engine_dir": engine_dir,
             }
 
-        self.file_path = get_absolute_config_path(file_path)
         self.init_switchboard_settings()
         self.init_sblhelper_settings()
         self.init_project_settings(basic_project_settings | p4_settings)
