SB: Changed packaged game checkbox to a drop down.

Also added more validations:
* Connected nodes
* Non-empty destination path
* Packaging program already running.

#rnx

diff --git a/Engine/Plugins/VirtualProduction/Switchboard/Source/Switchboard/switchboard/devices/ndisplay/plugin_ndisplay.py b/Engine/Plugins/VirtualProduction/Switchboard/Source/Switchboard/switchboard/devices/ndisplay/plugin_ndisplay.py
--- a/Engine/Plugins/VirtualProduction/Switchboard/Source/Switchboard/switchboard/devices/ndisplay/plugin_ndisplay.py
+++ b/Engine/Plugins/VirtualProduction/Switchboard/Source/Switchboard/switchboard/devices/ndisplay/plugin_ndisplay.py
@@ -43,6 +43,12 @@ class PackagingClientConfig(Enum):
     Development = "Development"
 
 
+class LaunchMode(Enum):
+    ''' Used to specify how to launch the Unreal instances. '''
+
+    Standalone = "Standalone (ICVFX)"
+    Packaged = "Packaged Game"
+
 class AddnDisplayDialog(AddDeviceDialog):
     def __init__(self, existing_devices, parent=None):
         super().__init__(
@@ -86,26 +92,27 @@ class AddnDisplayDialog(AddDeviceDialog):
         grid_layout.addWidget(lblConfigFile, 0, 0)
         grid_layout.addLayout(config_file_layout, 0, 1)
 
-        # Packaged Game row
+        # Launch As row
 
-        lblPackagedGame = QtWidgets.QLabel(self, text="Packaged Game")
-        self.chkPackagedGame = QtWidgets.QCheckBox(self)
-        self.chkPackagedGame.stateChanged.connect(self.on_chkPackagedGame_stateChanged)
+        lblLaunchAs = QtWidgets.QLabel(self, text="Launch As")
+        self.cmbLaunchAs = QtWidgets.QComboBox(self)
+        self.cmbLaunchAs.addItems([LaunchMode.Standalone.value, LaunchMode.Packaged.value])
+        self.cmbLaunchAs.currentIndexChanged.connect(self.on_cmbLaunchAs_currentIndexChanged)
 
-        packaged_game_layout = QtWidgets.QHBoxLayout()
-        packaged_game_layout.addWidget(self.chkPackagedGame)
-        packaged_game_layout.setAlignment(QtCore.Qt.AlignLeft)
+        launch_as_layout = QtWidgets.QHBoxLayout()
+        launch_as_layout.addWidget(self.cmbLaunchAs)
+        launch_as_layout.setAlignment(QtCore.Qt.AlignLeft)
 
-        grid_layout.addWidget(lblPackagedGame, 1, 0)
-        grid_layout.addLayout(packaged_game_layout, 1, 1)
+        grid_layout.addWidget(lblLaunchAs, 1, 0)
+        grid_layout.addLayout(launch_as_layout, 1, 1)
 
-        # Path row (initially hidden)
+        # Path row (initially not editable)
         self.lblPath = QtWidgets.QLabel(self, text=DevicenDisplay.csettings['packaged_game_path'].nice_name)
-        self.lblPath.setVisible(False)
+        self.lblPath.setEnabled(False)
         self.pathField = QtWidgets.QLineEdit(self)
-        self.pathField.setVisible(False)
+        self.pathField.setEnabled(False)
         self.btnPathBrowse = QtWidgets.QPushButton(self, text="...")
-        self.btnPathBrowse.setVisible(False)
+        self.btnPathBrowse.setEnabled(False)
         self.btnPathBrowse.clicked.connect(self.on_clicked_btnPathBrowse)
 
         path_layout = QtWidgets.QHBoxLayout()
@@ -142,17 +149,16 @@ class AddnDisplayDialog(AddDeviceDialog):
         # populate the config combobox with the items last populated.
         self.recall_config_itemDatas()
 
-    def on_chkPackagedGame_stateChanged(self, state):
-        ''' Called when the user changes the state of the Packaged Game checkbox
-        It updates the visibility of the packaged game path selection accordingly.
+    def on_cmbLaunchAs_currentIndexChanged(self, state):
+        ''' Called when the user changes the state of the LaunchAs dropdown
+        It updates the editability of the packaged game path selection accordingly.
         '''
+        is_packaged = self.cmbLaunchAs.currentText() == LaunchMode.Packaged.value
 
-        is_checked = QtCore.Qt.CheckState(state) == QtCore.Qt.Checked
-
-        self.lblPath.setVisible(is_checked)
-        self.pathField.setVisible(is_checked)
-        self.pathField.setEnabled(is_checked)
-        self.btnPathBrowse.setVisible(is_checked)
+        self.lblPath.setEnabled(is_packaged)
+        self.pathField.setEnabled(is_packaged)
+        self.pathField.setEnabled(is_packaged)
+        self.btnPathBrowse.setEnabled(is_packaged)
 
     def on_clicked_btnPathBrowse(self):
         ''' Called to browse for the packaged game executable path.'''
@@ -222,7 +228,8 @@ class AddnDisplayDialog(AddDeviceDialog):
         if res == QtWidgets.QDialog.Accepted:
             config_path = self.current_config_path()
             DevicenDisplay.csettings['ndisplay_config_file'].update_value(config_path)
-            DevicenDisplay.csettings['is_packaged_game'].update_value(self.is_packaged_game())
+            launch_mode = LaunchMode.Packaged.value if self.is_packaged_game() else LaunchMode.Standalone.value
+            DevicenDisplay.csettings['launch_mode'].update_value(launch_mode)
             DevicenDisplay.csettings['packaged_game_path'].update_value(self.packaged_game_path())
 
         return res
@@ -367,6 +374,7 @@ class DisplayConfig(object):
         self.uasset_path = ''
 
 
+
 class DevicenDisplay(DeviceUnreal):
 
     add_device_dialog = AddnDisplayDialog
@@ -381,18 +389,23 @@ class DevicenDisplay(DeviceUnreal):
             is_read_only=True,
             category="General Settings",
         ),
-        'is_packaged_game': BoolSetting(
-            attr_name="is_packaged_game",
-            nice_name="Packaged Game",
-            value=False,
-            tool_tip="Check if launching a packaged game.",
+        'launch_mode': OptionSetting(
+            attr_name="launch_mode",
+            nice_name="Launch As",
+            value=LaunchMode.Standalone.value,
+            possible_values=[
+                LaunchMode.Standalone.value,
+                LaunchMode.Packaged.value,
+            ],
+            tool_tip=f"Select the mode in which to launch the cluster. '{LaunchMode.Packaged.value}' will use the Packaged Game Path,"
+            f" while '{LaunchMode.Standalone.value}' will launch the project using the UnrealEditor executable with -game in the command line.",
             category="General Settings",
         ),
         'packaged_game_path': FilePathSetting(
             attr_name="packaged_game_path",
             nice_name="Packaged Executable",
             value="",
-            tool_tip="Path to the nDisplay packaged game executable. Only used when 'Packaged Game' is checked.",
+            tool_tip=f"Path to the nDisplay packaged game executable. Only used when '{LaunchMode.Packaged.value}' launch mode is selected.",
             show_ui=True,
             file_path_filter="Programs and Scripts (*.exe;*.bat;*.sh);;All Files (*)",
             category="General Settings",
@@ -924,7 +937,7 @@ class DevicenDisplay(DeviceUnreal):
 
     def is_packaged_game(self) -> bool:
         ''' Returns True if this node is launching as a package game'''
-        return DevicenDisplay.csettings['is_packaged_game'].get_value()
+        return DevicenDisplay.csettings['launch_mode'].get_value() == LaunchMode.Packaged.value
 
     def get_packaged_game_path(self) -> str:
         ''' Returns the packaged game executable path '''
@@ -1341,14 +1354,22 @@ class DevicenDisplay(DeviceUnreal):
         ''' Selects this node as the primary node in the nDisplay cluster '''
         self.__class__.select_device_as_primary(self)
 
-    def package_game(self, clientconfig: PackagingClientConfig = PackagingClientConfig.Default) -> uuid.UUID:
+    def package_game(
+            self,
+            clientconfig: PackagingClientConfig = PackagingClientConfig.Default,
+            dryrun: bool = False) -> uuid.UUID:
         ''' Packages the game using default settings '''
 
+        outdir = Path(self.get_packaged_game_path()).parent
+
+        if outdir == Path(""):
+            raise FileNotFoundError("Invalid Packaged Game Path")
+
         program_name = "unreal_packagegame"
-        outdir = Path(self.get_packaged_game_path())
 
-        if outdir.is_file():
-            outdir = outdir.parent
+        # Make sure there isn't a packaging program already running.
+        if len(self.program_start_queue.running_programs_named(program_name)):
+            raise PermissionError("Packaging task already running")
 
         # @todo using the remote platform would be more correct.
 
@@ -1389,6 +1410,9 @@ class DevicenDisplay(DeviceUnreal):
         prog_path = str(Path(CONFIG.ENGINE_DIR.get_value(self.name)).parent / prog_exe)
         prog_args = ' '.join(args)
 
+        if dryrun:
+            return uuid.uuid4()
+
         LOGGER.info(f'Packaging "{self.name}" with command: {prog_path} {prog_args}')
 
         puuid, msg = message_protocol.create_start_process_message(
diff --git a/Engine/Plugins/VirtualProduction/Switchboard/Source/Switchboard/switchboard/switchboard_dialog.py b/Engine/Plugins/VirtualProduction/Switchboard/Source/Switchboard/switchboard/switchboard_dialog.py
--- a/Engine/Plugins/VirtualProduction/Switchboard/Source/Switchboard/switchboard/switchboard_dialog.py
+++ b/Engine/Plugins/VirtualProduction/Switchboard/Source/Switchboard/switchboard/switchboard_dialog.py
@@ -617,11 +617,11 @@ class SwitchboardDialog(QtCore.QObject):
         self.insights_launcher_menuitem.setEnabled(not self.insights_launcher.is_running())
 
     def update_package_game_menuitem(self):
-        ''' Enables/disables the package game menu depending on whether
-        there are any nDisplay devices connected or not.
+        ''' Enables/disables the package game menu depending on whether there are any nDisplay devices or not.
+        For discoverability of the feature, we don't check if they are connected. Otherwise the user might be
+        clueless as to why the option is disabled.
         '''
-        ndisplay_devices = [device for device in self.device_manager.devices()
-                            if isinstance(device, DevicenDisplay) and not device.is_disconnected]
+        ndisplay_devices = [device for device in self.device_manager.devices() if isinstance(device, DevicenDisplay)]
 
         do_enable = len(ndisplay_devices) > 0
 
@@ -740,14 +740,36 @@ class SwitchboardDialog(QtCore.QObject):
         all_levels_action = self.register_tools_menu_action("All Levels", ["Fill DDC (Prepare Shaders)"])
         all_levels_action.triggered.connect(fill_ddc_all_levels_action)
 
-    def register_package_game_menuitem(self):
+    def register_package_game_menuitem(self) -> None:
         ''' Registers a menu item to package the nDisplay cluster game'''
 
-        def package_nDisplay_game(clientconfig: PackagingClientConfig):
+        def package_nDisplay_game(clientconfig: PackagingClientConfig) -> None:
             ''' Callback to kick off the packaging process '''
-            ndisplay_devices = [device for device in self.device_manager.devices() 
+            ndisplay_devices = [device for device in self.device_manager.devices()
                                 if isinstance(device, DevicenDisplay) and not device.is_disconnected]
 
+            # We need at least 1 device connected
+            if not len(ndisplay_devices):
+                QtWidgets.QMessageBox.information(
+                    self.window,
+                    "Unable to start packaging",
+                    "Please connect to at least one nDisplay node and try again."
+                )
+                return
+
+            # Validate all the devices before starting to package
+            for device in ndisplay_devices:
+                try:
+                    device.package_game(clientconfig=clientconfig, dryrun=True)
+                except Exception as e:
+                    QtWidgets.QMessageBox.information(
+                        self.window,
+                        "Unable to start packaging",
+                        f"nDisplay '{device.name}': {str(e)}"
+                    )
+                    return
+
+            # Ok, we should be good to go.
             for device in ndisplay_devices:
                 device.package_game(clientconfig)
 
