SB: Fixed add nD device.

Now protects against trying to package on multiple clients with the same ip address.
Also added status update of the packaging process. When it starts and when it ends.

#rnx

diff --git a/Engine/Plugins/VirtualProduction/Switchboard/Source/Switchboard/switchboard/devices/device_base.py b/Engine/Plugins/VirtualProduction/Switchboard/Source/Switchboard/switchboard/devices/device_base.py
--- a/Engine/Plugins/VirtualProduction/Switchboard/Source/Switchboard/switchboard/devices/device_base.py
+++ b/Engine/Plugins/VirtualProduction/Switchboard/Source/Switchboard/switchboard/devices/device_base.py
@@ -36,6 +36,7 @@ class DeviceStatus(IntEnum):
     OPEN = auto()
     READY = auto()
     RECORDING = auto()
+    PACKAGING = auto()
 
 
 @unique
@@ -56,11 +57,15 @@ class DeviceQtHandler(QtCore.QObject):
     signal_device_sync_failed = QtCore.Signal(object)
     signal_device_is_recording_device_changed = QtCore.Signal(object, int)
 
-    # Signal parameters are device, step, percent
+    # Signal parameters are: device, step, percent
     signal_device_build_update = QtCore.Signal(object, str, str)
-    # Signal parameters are device, percent
+
+    # Signal parameters are: device, percent
     signal_device_sync_update = QtCore.Signal(object, str)
 
+    # Signal parameters are: device, step, percentstr
+    signal_device_package_update = QtCore.Signal(object, str, str)
+
 
 class Device(QtCore.QObject):
     # Falls back to the default AddDeviceDialog as specified in
diff --git a/Engine/Plugins/VirtualProduction/Switchboard/Source/Switchboard/switchboard/devices/ndisplay/ndisplay_monitor.py b/Engine/Plugins/VirtualProduction/Switchboard/Source/Switchboard/switchboard/devices/ndisplay/ndisplay_monitor.py
--- a/Engine/Plugins/VirtualProduction/Switchboard/Source/Switchboard/switchboard/devices/ndisplay/ndisplay_monitor.py
+++ b/Engine/Plugins/VirtualProduction/Switchboard/Source/Switchboard/switchboard/devices/ndisplay/ndisplay_monitor.py
@@ -7,13 +7,14 @@ import threading
 import traceback
 
 from PySide6 import QtCore
-from PySide6.QtCore import QModelIndex, Qt, QTimer, Q_ARG
+from PySide6.QtCore import Qt, QTimer, Q_ARG
 from PySide6.QtGui import QColor, QIcon, QStandardItemModel, QStandardItem
 
 from switchboard import message_protocol
 from switchboard.message_protocol import SyncStatusRequestFlags
 from switchboard.switchboard_logging import LOGGER
 from switchboard.devices.device_base import Device
+from switchboard.devices.unreal.plugin_unreal import UnrealJobs
 
 
 class nDisplayMonitor(QStandardItemModel):
@@ -356,7 +357,7 @@ class nDisplayMonitor(QStandardItemModel):
     def program_id_from_device(self, device: Device):
         ''' Returns the program id of the running nDisplay unreal instance '''
         try:
-            program_id = device.program_start_queue.running_puuids_named('unreal')[-1]
+            program_id = device.program_start_queue.running_puuids_named(UnrealJobs.Unreal.value)[-1]
         except IndexError:
             program_id = self.default_program_id()
 
@@ -626,7 +627,7 @@ class nDisplayMonitor(QStandardItemModel):
         # Window in focus or not
         if SyncStatusRequestFlags.PidInFocus in request_flags:
             data['InFocus'] = 'no'
-            for prg in device.program_start_queue.running_programs_named('unreal'):
+            for prg in device.program_start_queue.running_programs_named(UnrealJobs.Unreal.value):
                 if prg.pid and prg.pid == sync_status['pidInFocus']:
                     data['InFocus'] = 'yes'
                     break
@@ -636,7 +637,7 @@ class nDisplayMonitor(QStandardItemModel):
             layers = '\n'.join([layer for layer in sync_status['programLayers'][1:]])
             if 'DISABLEDXMAXIMIZEDWINDOWEDMODE' in layers:
                 data['FSO'] = 'no'
-            elif len(device.program_start_queue.running_programs_named('unreal')):
+            elif len(device.program_start_queue.running_programs_named(UnrealJobs.Unreal.value)):
                 data['FSO'] = 'yes'
             else:
                 data['FSO'] = 'n/a'
diff --git a/Engine/Plugins/VirtualProduction/Switchboard/Source/Switchboard/switchboard/devices/ndisplay/plugin_ndisplay.py b/Engine/Plugins/VirtualProduction/Switchboard/Source/Switchboard/switchboard/devices/ndisplay/plugin_ndisplay.py
--- a/Engine/Plugins/VirtualProduction/Switchboard/Source/Switchboard/switchboard/devices/ndisplay/plugin_ndisplay.py
+++ b/Engine/Plugins/VirtualProduction/Switchboard/Source/Switchboard/switchboard/devices/ndisplay/plugin_ndisplay.py
@@ -8,7 +8,7 @@ import socket
 import struct
 import traceback
 import uuid
-from typing import Optional, Dict, OrderedDict
+from typing import Optional
 from enum import Enum
 
 from PySide6 import QtCore
@@ -28,7 +28,7 @@ from switchboard.devices.unreal.uassetparser import UassetParser
 from switchboard.devices.device_base import DeviceStatus
 from switchboard.switchboard_logging import LOGGER
 from switchboard.sbcache import SBCache, Asset
-from switchboard.devices.unreal.plugin_unreal import ProgramStartQueueItem
+from switchboard.devices.unreal.plugin_unreal import ProgramStartQueueItem, UnrealJobs
 
 from .ndisplay_monitor_ui import nDisplayMonitorUI
 from .ndisplay_monitor import nDisplayMonitor
@@ -153,7 +153,7 @@ class AddnDisplayDialog(AddDeviceDialog):
         ''' Called when the user changes the state of the LaunchAs dropdown
         It updates the editability of the packaged game path selection accordingly.
         '''
-        is_packaged = self.cmbLaunchAs.currentText() == LaunchMode.Packaged.value
+        is_packaged = self.is_packaged_game()
 
         self.lblPath.setEnabled(is_packaged)
         self.pathField.setEnabled(is_packaged)
@@ -217,7 +217,7 @@ class AddnDisplayDialog(AddDeviceDialog):
 
     def is_packaged_game(self):
         ''' Returns true if this is a packaged game'''
-        return self.chkPackagedGame.checkState() == QtCore.Qt.Checked
+        return self.cmbLaunchAs.currentText() == LaunchMode.Packaged.value
 
     def packaged_game_path(self):
         ''' Returns the path to the packaged game'''
@@ -1360,16 +1360,30 @@ class DevicenDisplay(DeviceUnreal):
             dryrun: bool = False) -> uuid.UUID:
         ''' Packages the game using default settings '''
 
-        outdir = Path(self.get_packaged_game_path()).parent
+        # Determine the folder where to package the game.
+
+        outdir = Path(self.get_packaged_game_path())
+
+        # Note: is_file() can't work reliably if the file doesn't exist yet.
+        if outdir.exists():
+            if outdir.is_file():
+                outdir = outdir.parent
+        else:
+            # if it ends with a separator, we know it is a directory
+            is_surely_dir = str(outdir).endswith(Path().anchor)
+            could_be_file = not is_surely_dir
+            if could_be_file:
+                outdir = outdir.parent
 
         if outdir == Path(""):
             raise FileNotFoundError("Invalid Packaged Game Path")
 
-        program_name = "unreal_packagegame"
+        program_name = UnrealJobs.PackageGame.value
 
-        # Make sure there isn't a packaging program already running.
-        if len(self.program_start_queue.running_programs_named(program_name)):
-            raise PermissionError("Packaging task already running")
+        # Make sure there isn't any other program already running.
+        # @todo Should include un-started programs
+        if self.program_start_queue.running_programs_count():
+            raise PermissionError("Another task is already running. Please try again later.")
 
         # @todo using the remote platform would be more correct.
 
diff --git a/Engine/Plugins/VirtualProduction/Switchboard/Source/Switchboard/switchboard/devices/unreal/plugin_unreal.py b/Engine/Plugins/VirtualProduction/Switchboard/Source/Switchboard/switchboard/devices/unreal/plugin_unreal.py
--- a/Engine/Plugins/VirtualProduction/Switchboard/Source/Switchboard/switchboard/devices/unreal/plugin_unreal.py
+++ b/Engine/Plugins/VirtualProduction/Switchboard/Source/Switchboard/switchboard/devices/unreal/plugin_unreal.py
@@ -21,6 +21,7 @@ import threading
 from typing import Callable, Dict, Generator, Optional, Union
 import uuid
 import time
+from enum import Enum
 
 from PySide6 import QtCore, QtGui, QtWidgets
 
@@ -45,6 +46,23 @@ from .redeploy_dialog import RedeployListenerDialog
 from switchboard.sbcache import SBCache, Asset
 
 
+class UnrealJobs(Enum):
+    ''' Used to specify the type of job sent to the unreal client '''
+
+    BuildPrefix = 'build_'
+    BuildProject = 'build_project'
+
+    CstatPrefix = 'cstat'
+    CstatProject = 'cstat_project'
+    CstatEngine = 'cstat_engine'
+
+    Sync = 'sync'
+    BuildDDC = 'unreal_ddc'
+    Unreal = 'unreal'
+    Retrieve = 'retrieve'
+    PackageGame = 'package_game'
+
+
 class ProgramStartQueueItem:
     '''
     Item that holds the information for a program to start in the future.
@@ -320,6 +338,10 @@ class ProgramStartQueue:
         '''
         return self.rd_puuid_by_name[name]
 
+    def running_programs_count(self) -> int:
+        ''' Returns number of running programs '''
+        return sum(len(program_list) for program_list in self.rd_running_progs_by_name.values())
+
 
 class LiveLinkPresetSetting(Setting):
     ''' Container of the LiveLink Preset setting
@@ -1494,7 +1516,7 @@ class DeviceUnreal(Device):
         formatstring = "%change%"
         args = f'-F "{formatstring}" -c {client_name} cstat {p4_path}/...#have'
 
-        program_name = "cstat_project"
+        program_name = UnrealJobs.CstatProject.value
 
         working_dir = os.path.dirname(
             CONFIG.UPROJECT_PATH.get_value(self.name))
@@ -1541,7 +1563,7 @@ class DeviceUnreal(Device):
         formatstring = "%change%"
         args = f'-F "{formatstring}" -c {client_name} cstat {p4_path}/...#have'
 
-        program_name = "cstat_engine"
+        program_name = UnrealJobs.CstatEngine.value
 
         working_dir = os.path.dirname(
             CONFIG.UPROJECT_PATH.get_value(self.name))
@@ -1611,7 +1633,7 @@ class DeviceUnreal(Device):
                                'workspace is already running.')
                 return
 
-        program_name = 'sync'
+        program_name = UnrealJobs.Sync.value
 
         # check if it is already on its way:
         try:
@@ -1760,7 +1782,7 @@ class DeviceUnreal(Device):
                                'same workspace is already running.')
                 return
 
-        program_name = 'build_project'
+        program_name = UnrealJobs.BuildProject.value
 
         # check if it is already on its way:
         try:
@@ -1776,7 +1798,7 @@ class DeviceUnreal(Device):
         sync_puuid = None
 
         try:
-            sync_puuid = self.program_start_queue.puuid_from_name('sync')
+            sync_puuid = self.program_start_queue.puuid_from_name(UnrealJobs.Sync.value)
         except KeyError:
             pass
 
@@ -1866,7 +1888,7 @@ class DeviceUnreal(Device):
     def _queue_build(
             self, program_name_suffix: str, ubt_args: str,
             puuid_dependency: Optional[uuid.UUID] = None):
-        program_name = f"build_{program_name_suffix}"
+        program_name = f"{UnrealJobs.BuildPrefix.value}{program_name_suffix}"
 
         engine_path = CONFIG.ENGINE_DIR.get_value(self.name)
         if self.target_platform.lower().startswith('win'):
@@ -1909,7 +1931,7 @@ class DeviceUnreal(Device):
 
     def close(self, force=False):
         # This call only refers to "unreal" programs.
-        unreal_puuids = self.program_start_queue.running_puuids_named('unreal')
+        unreal_puuids = self.program_start_queue.running_puuids_named(UnrealJobs.Unreal.value)
 
         for unreal_puuid in unreal_puuids:
             _, msg = message_protocol.create_kill_process_message(unreal_puuid)
@@ -1922,7 +1944,7 @@ class DeviceUnreal(Device):
 
     def disable_fso(self):
         ''' Tries to disable Full Screen Optimizations on the remote UnrealEditor.exe executable '''
-        unreals = self.program_start_queue.running_programs_named('unreal')
+        unreals = self.program_start_queue.running_programs_named(UnrealJobs.Unreal.value)
 
         if not len(unreals):
             return
@@ -2194,7 +2216,7 @@ class DeviceUnreal(Device):
             self.generate_unreal_exe_path(),
             self.generate_unreal_command_line_args(map_name))
 
-    def fill_derived_data_cache(self, current_level_only=False, program_name="unreal_ddc"):
+    def fill_derived_data_cache(self, current_level_only=False, program_name=UnrealJobs.BuildDDC.value):
         platforms: list[str] = self.setting_ddc_build_platforms.get_value(self.name)
         if (len(platforms) == 0):
             LOGGER.error("No platforms selected for fill_derived_data_cache, canceling...")
@@ -2235,7 +2257,7 @@ class DeviceUnreal(Device):
 
         return puuid
 
-    def launch(self, map_name, program_name="unreal"):
+    def launch(self, map_name, program_name=UnrealJobs.Unreal.value):
         if map_name == DEFAULT_MAP_TEXT:
             map_name = ''
 
@@ -2308,15 +2330,17 @@ class DeviceUnreal(Device):
             self.device_qt_handler.signal_device_client_disconnected.emit(self)
 
     def do_program_running_update(self, prog):
-        if prog.name == 'unreal':
+        if prog.name == UnrealJobs.Unreal.value:
             self.status = DeviceStatus.OPEN
             self.unreal_started_signal.emit()
-        elif prog.name.startswith('build_'):
+        elif prog.name.startswith(UnrealJobs.BuildPrefix.value):
             for device in self.devices_sharing_workspace():
                 device.status = DeviceStatus.BUILDING
-        elif prog.name == 'sync':
+        elif prog.name == UnrealJobs.Sync.value:
             for device in self.devices_sharing_workspace():
                 device.status = DeviceStatus.SYNCING
+        elif prog.name == UnrealJobs.PackageGame.value:
+            self.status = DeviceStatus.PACKAGING
 
         self.program_start_queue.update_running_program(prog=prog)
 
@@ -2335,15 +2359,15 @@ class DeviceUnreal(Device):
             # log this
             LOGGER.error(f"Could not start {program_name}: {message['error']}")
 
-            if program_name == 'sync' or program_name.startswith('build_'):
+            if program_name == UnrealJobs.Sync.value or program_name.startswith(UnrealJobs.BuildPrefix.value):
                 for device in self.devices_sharing_workspace():
                     device.status = DeviceStatus.CLOSED
 
                     # This has the effect of hiding building/syncing status.
                     device.project_changelist = device.project_changelist
-            elif program_name == 'unreal':
+            elif program_name == UnrealJobs.Unreal.value:
                 self.status = DeviceStatus.CLOSED
-            elif program_name == 'retrieve':
+            elif program_name == UnrealJobs.Retrieve.value:
                 self.status = DeviceStatus.CLOSED
 
             return
@@ -2419,13 +2443,13 @@ class DeviceUnreal(Device):
         remaining_homonyms = self.program_start_queue.running_puuids_named(
             program_name)
         for prog_id in remaining_homonyms:
-            if program_name != 'retrieve':
+            if program_name != UnrealJobs.Retrieve.value:
                 LOGGER.warning(
                     f'{self.name}: But ({prog_id}) with the same name '
                     f'"{program_name}" is still in the list, which is unusual')
 
-        if program_name == 'unreal' and not len(remaining_homonyms):
-            
+        if program_name == UnrealJobs.Unreal.value and not len(remaining_homonyms):
+
             if DeviceUnreal.csettings["retrieve_logs"].get_value():
                 log_success = self.start_retrieve_log(unreal_exit_code=returncode)
                 utrace_success = self.start_retrieve_utrace(unreal_exit_code=returncode)
@@ -2434,10 +2458,10 @@ class DeviceUnreal(Device):
             else:
                 self.status = DeviceStatus.CLOSED
 
-        elif program_name == 'retrieve' and not self.transfer_in_progress:
+        elif program_name == UnrealJobs.Retrieve.value and not self.transfer_in_progress:
             self.status = DeviceStatus.CLOSED
 
-        elif program_name == 'sync':
+        elif program_name == UnrealJobs.Sync.value:
             if returncode == 0:
                 LOGGER.info(f"{self.name}: Sync successful")
                 self.last_sync_filter_hash.update_value(
@@ -2478,7 +2502,7 @@ class DeviceUnreal(Device):
                 other.project_changelist = self.project_changelist
                 other.status = DeviceStatus.CLOSED
 
-        elif program_name.startswith('build_'):
+        elif program_name.startswith(UnrealJobs.BuildPrefix.value):
             if returncode == 0:
                 LOGGER.info(f"{self.name}: {program_name} successful!")
             else:
@@ -2491,7 +2515,7 @@ class DeviceUnreal(Device):
                     if error_pattern.search(line):
                         LOGGER.error(f"{self.name}: {line}")
 
-            if 'build_project' == program_name:
+            if UnrealJobs.BuildProject.value == program_name:
                 for device in self.devices_sharing_workspace():
                     device.status = DeviceStatus.CLOSED
 
@@ -2502,7 +2526,7 @@ class DeviceUnreal(Device):
                     self._request_engine_changelist_number()
                     self._request_unreal_editor_version_file()
 
-        elif "cstat" in program_name:
+        elif UnrealJobs.CstatPrefix.value in program_name:
             output = get_stdout_str()
             changelists = [line.strip() for line in output.split()]
 
@@ -2539,6 +2563,12 @@ class DeviceUnreal(Device):
                     f"{project_name} is on revision {current_changelist}")
                 for device in self.devices_sharing_workspace():
                     device.engine_changelist = current_changelist
+        elif program_name == UnrealJobs.PackageGame.value:
+            self.status = DeviceStatus.CLOSED
+            percentstr = '100%'
+            if returncode != 0:
+                percentstr = 'fail'
+            self.widget.update_package_status(device=self, step='', percentstr=percentstr)
 
     def on_program_killed(self, message):
         '''
@@ -2621,7 +2651,7 @@ class DeviceUnreal(Device):
         #    @progress 'Generating code...' 100%
         #    @progress pop
         #
-        if process['name'].startswith('build_'):
+        if process['name'].startswith(UnrealJobs.BuildPrefix.value):
             for line in lines:
                 if '@progress' not in line:
                     continue
@@ -2641,7 +2671,7 @@ class DeviceUnreal(Device):
                     device.device_qt_handler.signal_device_build_update.emit(
                             device, step, percent)
 
-        elif process['name'] == 'sync':
+        elif process['name'] == UnrealJobs.Sync.value:
             for line in lines:
                 if 'Progress:' not in line:
                     continue
@@ -2655,6 +2685,12 @@ class DeviceUnreal(Device):
                     device.device_qt_handler.signal_device_sync_update.emit(
                         device, sync_progress)
 
+        elif process['name'] == UnrealJobs.PackageGame.value:
+            step = ''
+            sync_progress = ''
+            self.device_qt_handler.signal_device_package_update.emit(
+                        self, step, sync_progress)
+
         for line in lines:
             LOGGER.debug(f"{self.name} {process['name']}: {line}")
 
@@ -2844,7 +2880,7 @@ class DeviceUnreal(Device):
             return 'rsync'
 
     def fetch_file(self, remote_path):
-        program_name = 'retrieve'
+        program_name = UnrealJobs.Retrieve.value
 
         rsync_path = self.get_rsync_path()
 
@@ -2928,7 +2964,7 @@ class DeviceUnreal(Device):
     @property
     def transfer_in_progress(self) -> bool:
         return len(
-            self.program_start_queue.running_puuids_named('retrieve')) > 0
+            self.program_start_queue.running_puuids_named(UnrealJobs.Retrieve.value)) > 0
 
     def get_widget_classes(self):
         widget_classes = list(self._widget_classes)
@@ -3448,6 +3484,12 @@ class DeviceWidgetUnreal(DeviceWidget):
             self.open_button.setChecked(True)
             self.sync_button.setDisabled(True)
             self.build_button.setDisabled(True)
+        elif status == DeviceStatus.PACKAGING:
+            self.open_button.setDisabled(True)
+            self.sync_button.setDisabled(True)
+            self.build_button.setDisabled(True)
+            self.engine_changelist_label.hide()
+            self.project_changelist_label.setText('Packaging...')
 
         if self.open_button.isChecked():
             if self.open_button.isEnabled():
@@ -3481,6 +3523,24 @@ class DeviceWidgetUnreal(DeviceWidget):
 
         self.project_changelist_label.show()
 
+    @QtCore.Slot(Device, str, str)
+    def update_package_status(self, device: Device, step: str, percentstr: str):
+        if threading.current_thread() is not threading.main_thread():
+            QtCore.QMetaObject.invokeMethod(
+                self,
+                'update_package_status',
+                QtCore.Qt.QueuedConnection,
+                QtCore.Q_ARG(Device, device),
+                QtCore.Q_ARG(str, step),
+                QtCore.Q_ARG(str, percentstr)
+            )
+            return
+
+        self.project_changelist_label.setText(f'Packaging...{percentstr}')
+        self.project_changelist_label.setToolTip(step)
+
+        self.project_changelist_label.show()
+
     def update_engine_changelist(self, required_cl: str, synched_cl: str, built__cl: str):
         if not CONFIG.ENGINE_SYNC_METHOD.get_value() == EngineSyncMethod.Build_Engine.value:
             return
diff --git a/Engine/Plugins/VirtualProduction/Switchboard/Source/Switchboard/switchboard/switchboard_dialog.py b/Engine/Plugins/VirtualProduction/Switchboard/Source/Switchboard/switchboard/switchboard_dialog.py
--- a/Engine/Plugins/VirtualProduction/Switchboard/Source/Switchboard/switchboard/switchboard_dialog.py
+++ b/Engine/Plugins/VirtualProduction/Switchboard/Source/Switchboard/switchboard/switchboard_dialog.py
@@ -757,6 +757,18 @@ class SwitchboardDialog(QtCore.QObject):
                 )
                 return
 
+            # Make sure we won't try to send the package command to the same client multiple times.
+
+            unique_addresses = {device.address for device in ndisplay_devices}
+
+            if len(ndisplay_devices) > len(unique_addresses):
+                QtWidgets.QMessageBox.information(
+                    self.window,
+                    "Unable to start packaging",
+                    "Please only try to package devices with unique addresses."
+                )
+                return
+
             # Validate all the devices before starting to package
             for device in ndisplay_devices:
                 try:
@@ -771,7 +783,7 @@ class SwitchboardDialog(QtCore.QObject):
 
             # Ok, we should be good to go.
             for device in ndisplay_devices:
-                device.package_game(clientconfig)
+                device.package_game(clientconfig=clientconfig)
 
         actionname = "Package nDisplay Game"
 
@@ -1440,6 +1452,7 @@ class SwitchboardDialog(QtCore.QObject):
         device.device_qt_handler.signal_device_is_recording_device_changed.connect(self.device_is_recording_device_changed, QtCore.Qt.QueuedConnection)
         device.device_qt_handler.signal_device_build_update.connect(self.device_build_update, QtCore.Qt.QueuedConnection)
         device.device_qt_handler.signal_device_sync_update.connect(self.device_sync_update, QtCore.Qt.QueuedConnection)
+        device.device_qt_handler.signal_device_package_update.connect(self.device_package_update, QtCore.Qt.QueuedConnection)
 
         # Add the view
         self.device_list_widget.add_device_widget(device)
@@ -1839,6 +1852,11 @@ class SwitchboardDialog(QtCore.QObject):
         device_widget = self.device_list_widget.device_widget_by_hash(device.device_hash)
         device_widget.update_sync_status(device, progress)
 
+    @QtCore.Slot(object)
+    def device_package_update(self, device, step, percent):
+        device_widget = self.device_list_widget.device_widget_by_hash(device.device_hash)
+        device_widget.update_package_status(device, step, percent)
+
     @QtCore.Slot(object)
     def device_project_changelist_changed(self, device):
         device_widget = self.device_list_widget.device_widget_by_hash(device.device_hash)
