Switchboard: Avoid trying to send a cluster event to a cluster that is not connected or running.

This avoids unnecessary network connection delays that are not likely to succeed.

#rnx

diff --git a/Engine/Plugins/VirtualProduction/Switchboard/Source/Switchboard/switchboard/devices/ndisplay/plugin_ndisplay.py b/Engine/Plugins/VirtualProduction/Switchboard/Source/Switchboard/switchboard/devices/ndisplay/plugin_ndisplay.py
--- a/Engine/Plugins/VirtualProduction/Switchboard/Source/Switchboard/switchboard/devices/ndisplay/plugin_ndisplay.py
+++ b/Engine/Plugins/VirtualProduction/Switchboard/Source/Switchboard/switchboard/devices/ndisplay/plugin_ndisplay.py
@@ -8,7 +8,7 @@ import socket
 import struct
 import traceback
 import uuid
-from typing import Optional
+from typing import Optional, List
 from enum import Enum
 
 from PySide6 import QtCore
@@ -1841,7 +1841,7 @@ class DevicenDisplay(DeviceUnreal):
         cls.ndisplay_monitor.removed_device(device)
 
     @classmethod
-    def send_cluster_event(cls, devices, cluster_event):
+    def send_cluster_event(cls, devices: List[DeviceUnreal], cluster_event):
         '''
         Sends a cluster event (to the primary node, which will replicate to the rest
         of the cluster).
@@ -1864,6 +1864,14 @@ class DevicenDisplay(DeviceUnreal):
                            'device is marked as such.')
             raise ValueError
 
+        if not primary.is_connected_and_authenticated():
+            LOGGER.warning('Primary node must be connected and running to send a cluster event')
+            raise ValueError
+
+        if not primary.program_start_queue.running_programs_named(UnrealJobs.Unreal.value):
+            LOGGER.warning('Primary node must be connected and running to send a cluster event')
+            raise ValueError
+
         msg = bytes(json.dumps(cluster_event), 'utf-8')
         msg = struct.pack('I', len(msg)) + msg
 
