Switchboard: Ensure newly created configs properly update the plugin manager engine/project paths.

Fixes an exception and a potential hang when creating new configs.

#rnx

diff --git a/Engine/Plugins/VirtualProduction/Switchboard/Source/Switchboard/switchboard/config.py b/Engine/Plugins/VirtualProduction/Switchboard/Source/Switchboard/switchboard/config.py
--- a/Engine/Plugins/VirtualProduction/Switchboard/Source/Switchboard/switchboard/config.py
+++ b/Engine/Plugins/VirtualProduction/Switchboard/Source/Switchboard/switchboard/config.py
@@ -2025,6 +2025,7 @@ class Config(object):
         self.init_switchboard_settings(data)
         self.init_sblhelper_settings(data)
         self.init_project_settings(data)
+        self.init_plugin_tracking()
         self.init_unreal_insights(data)
         self.init_muserver(data)
 
@@ -2048,7 +2049,6 @@ class Config(object):
 
         self.CURRENT_LEVEL = data.get('current_level', DEFAULT_MAP_TEXT)
 
-        self.init_plugin_tracking()
         self.init_devices(data)
 
     def _backup_corrupted_config(self, original_file_path: str):
@@ -2104,9 +2104,12 @@ class Config(object):
     def init_plugin_tracking(self):
         ''' Initializes the plugin manager '''
 
+        engine_dir = self.ENGINE_DIR.get_value()
+        uproj_path = self.UPROJECT_PATH.get_value()
+
         self.ue_plugin_mgr = UnrealPluginManager()
-        self.ue_plugin_mgr.set_engine_dir(Path(self.ENGINE_DIR.get_value()))
-        self.ue_plugin_mgr.set_uproject_path(Path(self.UPROJECT_PATH.get_value()))
+        self.ue_plugin_mgr.set_engine_dir(Path(engine_dir) if engine_dir else None)
+        self.ue_plugin_mgr.set_uproject_path(Path(uproj_path) if uproj_path else None)
 
         self.ENGINE_DIR.signal_setting_changed.connect(
             lambda _, new: self.ue_plugin_mgr.set_engine_dir(Path(new)))
@@ -2129,6 +2132,7 @@ class Config(object):
         self.init_switchboard_settings()
         self.init_sblhelper_settings()
         self.init_project_settings(basic_project_settings | p4_settings)
+        self.init_plugin_tracking()
         self.init_unreal_insights()
         self.init_muserver()
 
