Switchboard: Project CL sync is now disabled if an empty p4 path is specified.

The objective is that when the project is not in the same depot as the Engine, that SB doesn't know invalid CLs in the project dropdown or try to sync them.

diff --git a/Engine/Plugins/VirtualProduction/Switchboard/Source/Switchboard/switchboard/switchboard_dialog.py b/Engine/Plugins/VirtualProduction/Switchboard/Source/Switchboard/switchboard/switchboard_dialog.py
--- a/Engine/Plugins/VirtualProduction/Switchboard/Source/Switchboard/switchboard/switchboard_dialog.py
+++ b/Engine/Plugins/VirtualProduction/Switchboard/Source/Switchboard/switchboard/switchboard_dialog.py
@@ -2055,18 +2055,31 @@ class SwitchboardDialog(QtCore.QObject):
 
         sync_method = CONFIG.ENGINE_SYNC_METHOD.get_value()
 
+        # If we're syncing 'Precompiled Binaries', then that implies that we should be using UGS:
+        use_ugs = ENABLE_UGS_SUPPORT and (sync_method in (EngineSyncMethod.Sync_PCBs.value, EngineSyncMethod.Sync_From_UGS.value))
+
+        # Disable project sync options if not using UGS and there isn't a project p4 path set.
+        if CONFIG.P4_PROJECT_PATH.get_value() or use_ugs:
+            self.window.project_cl_label.setEnabled(True)
+            self.window.project_cl_combo_box.setEnabled(True)
+            self.window.refresh_project_cl_button.setEnabled(True)
+        else:
+            self.window.project_cl_label.setEnabled(False)
+            self.window.project_cl_combo_box.setEnabled(False)
+            self.window.refresh_project_cl_button.setEnabled(False)
+            return
+
         changelists: Optional[list[int]] = None
         descriptions: Optional[list[str]] = None
 
-        # If we're syncing 'Precompiled Binaries', then that implies that we should be using UGS:
-        if ENABLE_UGS_SUPPORT:
-            if sync_method == EngineSyncMethod.Sync_PCBs.value or sync_method == EngineSyncMethod.Sync_From_UGS.value:
-                LOGGER.info("Using UnrealGameSync to refresh project changelists.")
-                changelists = ugs_utils.latest_chagelists(Path(CONFIG.UPROJECT_PATH.get_value()), client=CONFIG.SOURCE_CONTROL_WORKSPACE.get_value())
-                if not changelists:
-                    LOGGER.error("UnrealGameSync failed to get the project's latest changelists. Falling back to using p4 commands directly.")
+        if use_ugs:
+            LOGGER.info("Using UnrealGameSync to refresh project changelists.")
+            changelists = ugs_utils.latest_chagelists(Path(CONFIG.UPROJECT_PATH.get_value()), client=CONFIG.SOURCE_CONTROL_WORKSPACE.get_value())
+            if not changelists:
+                LOGGER.error("UnrealGameSync failed to get the project's latest changelists. Falling back to using p4 commands directly.")
 
         if not changelists:
+
             LOGGER.info("Refreshing p4 project changelists")
             client = CONFIG.SOURCE_CONTROL_WORKSPACE.get_value()
             paths = [f'{CONFIG.P4_PROJECT_PATH.get_value()}/...']
