SB: Support separate engine and project workspaces.

This should allow users to launch projects that are on different streams and workspaces than the Engine itself.

The UI addition is a way to specify a project workspace, when applicable.

[RN] Switchboard now supports separate engine and project Perforce workspaces

diff --git a/Engine/Plugins/VirtualProduction/Switchboard/Source/Switchboard/switchboard/add_config_dialog.py b/Engine/Plugins/VirtualProduction/Switchboard/Source/Switchboard/switchboard/add_config_dialog.py
--- a/Engine/Plugins/VirtualProduction/Switchboard/Source/Switchboard/switchboard/add_config_dialog.py
+++ b/Engine/Plugins/VirtualProduction/Switchboard/Source/Switchboard/switchboard/add_config_dialog.py
@@ -34,13 +34,14 @@ class AddConfigDialog(QtWidgets.QDialog):
             Args:
                 addConfigDialog(AddConfigDialog): The dialog to load the data from
             '''
-            self.config_path_str  = addConfigDialog.config_path_line_edit.text()
-            self.engineDir        = addConfigDialog.engine_dir_line_edit.text()
-            self.project          = addConfigDialog.uproject_line_edit.text()
-            self.p4Engine         = addConfigDialog.p4_engine_path_line_edit.text()
-            self.p4Project        = addConfigDialog.p4_project_path_line_edit.text()
-            self.p4Workspace      = addConfigDialog.p4_workspace_line_edit.text()
-            self.p4Enabled        = addConfigDialog.p4_group.isChecked()
+            self.config_path_str    = addConfigDialog.config_path_line_edit.text()
+            self.engineDir          = addConfigDialog.engine_dir_line_edit.text()
+            self.project            = addConfigDialog.uproject_line_edit.text()
+            self.p4Engine           = addConfigDialog.p4_engine_path_line_edit.text()
+            self.p4Project          = addConfigDialog.p4_project_path_line_edit.text()
+            self.p4Workspace        = addConfigDialog.p4_workspace_line_edit.text()
+            self.p4ProjectWorkspace = addConfigDialog.p4_project_workspace_line_edit.text()
+            self.p4Enabled          = addConfigDialog.p4_group.isChecked()
 
         def dump(self, addConfigDialog):
             ''' Writes the relevant fields from AddConfigDialog
@@ -48,12 +49,13 @@ class AddConfigDialog(QtWidgets.QDialog):
                 addConfigDialog(AddConfigDialog): The dialog to load the data from
             '''
             lineEdits = [
-                (self.config_path_str, addConfigDialog.config_path_line_edit    ),
-                (self.engineDir      , addConfigDialog.engine_dir_line_edit     ),
-                (self.project        , addConfigDialog.uproject_line_edit       ),
-                (self.p4Engine       , addConfigDialog.p4_engine_path_line_edit ),
-                (self.p4Project      , addConfigDialog.p4_project_path_line_edit),
-                (self.p4Workspace    , addConfigDialog.p4_workspace_line_edit   ),
+                (self.config_path_str   , addConfigDialog.config_path_line_edit         ),
+                (self.engineDir         , addConfigDialog.engine_dir_line_edit          ),
+                (self.project           , addConfigDialog.uproject_line_edit            ),
+                (self.p4Engine          , addConfigDialog.p4_engine_path_line_edit      ),
+                (self.p4Project         , addConfigDialog.p4_project_path_line_edit     ),
+                (self.p4Workspace       , addConfigDialog.p4_workspace_line_edit        ),
+                (self.p4ProjectWorkspace, addConfigDialog.p4_project_workspace_line_edit),
             ]
 
             # Update line edits
@@ -143,10 +145,14 @@ class AddConfigDialog(QtWidgets.QDialog):
         self.p4_workspace_line_edit = QtWidgets.QLineEdit(
             config.CONFIG.SOURCE_CONTROL_WORKSPACE.get_value())
 
+        self.p4_project_workspace_line_edit = QtWidgets.QLineEdit(
+            config.CONFIG.PROJECT_WORKSPACE.get_value())
+
         p4_layout = QtWidgets.QFormLayout()
         p4_layout.addRow("P4 Project Path", self.p4_project_path_line_edit)
         p4_layout.addRow("P4 Engine Path", self.p4_engine_path_line_edit)
-        p4_layout.addRow("Workspace Name", self.p4_workspace_line_edit)
+        p4_layout.addRow("Workspace", self.p4_workspace_line_edit)
+        p4_layout.addRow("Project Workspace Override", self.p4_project_workspace_line_edit)
         self.p4_group.setLayout(p4_layout)
         layout.addWidget(self.p4_group)
 
@@ -504,6 +510,7 @@ class AddConfigDialog(QtWidgets.QDialog):
         settings = {}
         settings['p4_enabled'] = self.p4_group.isChecked()
         settings['source_control_workspace'] = self.p4_workspace_line_edit.text() if self.p4_group.isChecked() else None
+        settings['project_workspace'] = self.p4_project_workspace_line_edit.text() if self.p4_group.isChecked() else None
         settings['p4_sync_path'] = self.p4_project_path_line_edit.text() if self.p4_group.isChecked() else None
         settings['p4_engine_path'] = self.p4_engine_path_line_edit.text() if self.p4_group.isChecked() else None
         return settings
diff --git a/Engine/Plugins/VirtualProduction/Switchboard/Source/Switchboard/switchboard/config.py b/Engine/Plugins/VirtualProduction/Switchboard/Source/Switchboard/switchboard/config.py
--- a/Engine/Plugins/VirtualProduction/Switchboard/Source/Switchboard/switchboard/config.py
+++ b/Engine/Plugins/VirtualProduction/Switchboard/Source/Switchboard/switchboard/config.py
@@ -2293,11 +2293,17 @@ class Config(object):
                 category="Source Control Settings",
             ),
             "source_control_workspace": StringSetting(
-                "source_control_workspace", "Workspace Name",
-                data.get("source_control_workspace"),
+                "source_control_workspace", "Workspace",
+                data.get("source_control_workspace", ""),
                 tool_tip="SourceControl Workspace/Branch",
                 category="Source Control Settings",
             ),
+            "project_workspace": StringSetting(
+                "project_workspace", "Project Workspace Override",
+                data.get("project_workspace", ""),
+                tool_tip="Only use if different to the Engine workspace",
+                category="Source Control Settings",
+            ),
             "p4_sync_path": PerforcePathSetting(
                 "p4_sync_path",
                 "Perforce Project Path",
@@ -2316,6 +2322,7 @@ class Config(object):
 
         self.P4_ENABLED = self.source_control_settings["p4_enabled"]
         self.SOURCE_CONTROL_WORKSPACE = self.source_control_settings["source_control_workspace"]
+        self.PROJECT_WORKSPACE = self.source_control_settings["project_workspace"]
         self.P4_PROJECT_PATH: PerforcePathSetting = self.source_control_settings["p4_sync_path"]
         self.P4_ENGINE_PATH: PerforcePathSetting = self.source_control_settings["p4_engine_path"]
 
@@ -2569,8 +2576,8 @@ class Config(object):
         data["p4_enabled"] = self.P4_ENABLED.get_value()
         data["p4_sync_path"] = self.P4_PROJECT_PATH.get_value()
         data["p4_engine_path"] = self.P4_ENGINE_PATH.get_value()
-        data["source_control_workspace"] = (
-            self.SOURCE_CONTROL_WORKSPACE.get_value())
+        data["source_control_workspace"] = (self.SOURCE_CONTROL_WORKSPACE.get_value())
+        data["project_workspace"] = (self.PROJECT_WORKSPACE.get_value())
 
         self.save_muserver(data)
 
diff --git a/Engine/Plugins/VirtualProduction/Switchboard/Source/Switchboard/switchboard/devices/ndisplay/plugin_ndisplay.py b/Engine/Plugins/VirtualProduction/Switchboard/Source/Switchboard/switchboard/devices/ndisplay/plugin_ndisplay.py
--- a/Engine/Plugins/VirtualProduction/Switchboard/Source/Switchboard/switchboard/devices/ndisplay/plugin_ndisplay.py
+++ b/Engine/Plugins/VirtualProduction/Switchboard/Source/Switchboard/switchboard/devices/ndisplay/plugin_ndisplay.py
@@ -820,6 +820,7 @@ class DevicenDisplay(DeviceUnreal):
             DevicenDisplay.csettings['packaged_game_path'],
             CONFIG.ENGINE_DIR,
             CONFIG.SOURCE_CONTROL_WORKSPACE,
+            CONFIG.PROJECT_WORKSPACE,
             CONFIG.UPROJECT_PATH,
         ]
 
diff --git a/Engine/Plugins/VirtualProduction/Switchboard/Source/Switchboard/switchboard/devices/unreal/plugin_unreal.py b/Engine/Plugins/VirtualProduction/Switchboard/Source/Switchboard/switchboard/devices/unreal/plugin_unreal.py
--- a/Engine/Plugins/VirtualProduction/Switchboard/Source/Switchboard/switchboard/devices/unreal/plugin_unreal.py
+++ b/Engine/Plugins/VirtualProduction/Switchboard/Source/Switchboard/switchboard/devices/unreal/plugin_unreal.py
@@ -1261,6 +1261,7 @@ class DeviceUnreal(Device):
             DeviceUnreal.csettings['mediaprofile'],
             CONFIG.ENGINE_DIR,
             CONFIG.SOURCE_CONTROL_WORKSPACE,
+            CONFIG.PROJECT_WORKSPACE,
             CONFIG.UPROJECT_PATH,
         ]
 
@@ -1498,6 +1499,10 @@ class DeviceUnreal(Device):
             return
 
         client_name = CONFIG.SOURCE_CONTROL_WORKSPACE.get_value(self.name)
+        project_workspace = CONFIG.PROJECT_WORKSPACE.get_value(self.name)
+
+        if project_workspace:
+            client_name = project_workspace
 
         if not client_name:
             LOGGER.warning(
@@ -1619,35 +1624,48 @@ class DeviceUnreal(Device):
         super().disconnect_listener()
         self.unreal_client.disconnect()
 
-    def sync(self, engine_cl: Optional[int], project_cl: Optional[int]):
+    def sync(
+            self,
+            engine_cl: Optional[int],
+            project_cl: Optional[int],
+            puuid_dependency: Optional[uuid.UUID] = None
+            ) -> Optional[uuid.UUID]:
+
         if not engine_cl and not project_cl:
             LOGGER.warning(
                 "Neither project nor engine changelist is selected. "
                 "There is nothing to sync!")
-            return
+            return None
 
         for device in self.devices_sharing_workspace():
             if device.status >= DeviceStatus.OPEN:
                 LOGGER.warning(f'{self.name}: Unable to sync, because another '
                                f'device ({device.name}) using the same '
                                'workspace is already running.')
-                return
-
-        program_name = UnrealJobs.Sync.value
-
-        # check if it is already on its way:
-        try:
-            existing_puuid = self.program_start_queue.puuid_from_name(
-                program_name)
-            LOGGER.info(
-                f"{self.name}: Already syncing with puuid {existing_puuid}")
-            return
-        except KeyError:
-            pass
+                return None
 
         project_path = CONFIG.UPROJECT_PATH.get_value(self.name)
         engine_dir = CONFIG.ENGINE_DIR.get_value(self.name)
         workspace = CONFIG.SOURCE_CONTROL_WORKSPACE.get_value(self.name)
+        project_workspace = CONFIG.PROJECT_WORKSPACE.get_value(self.name)
+
+        # If only we're syncing the project and the project has a specified workspace, use the project workspace
+        # instead.
+        if not engine_cl and project_cl and project_workspace:
+            workspace = project_workspace
+
+        program_name = UnrealJobs.Sync.value
+
+        # If there are no task dependencies, check if it is already on its way:
+        if not puuid_dependency:
+            try:
+                existing_puuid = self.program_start_queue.puuid_from_name(
+                    program_name)
+                LOGGER.info(
+                    f"{self.name}: Already syncing with puuid {existing_puuid}")
+                return None
+            except KeyError:
+                pass
 
         sync_method = CONFIG.ENGINE_SYNC_METHOD.get_value()
         generate_proj_files = sync_method == EngineSyncMethod.Build_Engine.value
@@ -1674,12 +1692,18 @@ class DeviceUnreal(Device):
             'Source', 'Switchboard', 'sbl_helper.py'))
 
         sync_tool = python_path
-        sync_args = (
-            f'"{helper_path}" '
-            # '--log-level=DEBUG '
-            'sync '
-            f'--project="{project_path}" '
-            f'--engine-dir="{engine_dir}"')
+
+        sync_args = [
+            f'"{helper_path}"',
+            # '--log-level=DEBUG',
+            'sync',
+            f'--project="{project_path}"',
+        ]
+
+        if engine_cl:
+            sync_args.append(f'--engine-dir="{engine_dir}"')
+
+        sync_args = ' '.join(sync_args)
 
         if workspace:
             sync_args += f' --p4client={workspace}'
@@ -1687,10 +1711,14 @@ class DeviceUnreal(Device):
         if engine_cl:
             sync_args += f' --engine-cl={engine_cl}'
             self.inflight_engine_cl = engine_cl
+        else:
+            self.inflight_engine_cl = None
 
         if project_cl:
             sync_args += f' --project-cl={project_cl} --clobber-project'
             self.inflight_project_cl = project_cl
+        else:
+            self.inflight_project_cl = None
 
         if generate_proj_files:
             sync_args += ' --generate'
@@ -1738,8 +1766,8 @@ class DeviceUnreal(Device):
             if sync_precompiled_bins:
                 sync_args += ' --use-pcbs'
 
-                # If we're syncing 'Precompiled Binaries' one of those binaries may be the SwitchboardListener executable
-                # which means (on Windows atleast) we need to move the executable to make way for the new one
+                # If we're syncing 'Precompiled Binaries' one of those binaries may be the SwitchboardListener
+                # executable which means (on Windows atleast) we need to move the executable to make way for the new one
                 _, msg = message_protocol.create_free_listener_bin_message()
                 self.unreal_client.send_message(msg)
 
@@ -1757,7 +1785,7 @@ class DeviceUnreal(Device):
         self.program_start_queue.add(
             ProgramStartQueueItem(
                 name=program_name,
-                puuid_dependency=None,
+                puuid_dependency=puuid_dependency,
                 puuid=puuid,
                 msg_to_unreal_client=msg,
                 launch_fn=lambda: LOGGER.info(
@@ -1771,6 +1799,8 @@ class DeviceUnreal(Device):
             if device.status != DeviceStatus.SYNCING:
                 device.status = DeviceStatus.SYNCING
 
+        return puuid
+
     def build(self):
         if self.exclude_from_build.get_value():
             return
@@ -2331,14 +2361,23 @@ class DeviceUnreal(Device):
 
     def do_program_running_update(self, prog):
         if prog.name == UnrealJobs.Unreal.value:
+
             self.status = DeviceStatus.OPEN
             self.unreal_started_signal.emit()
+
         elif prog.name.startswith(UnrealJobs.BuildPrefix.value):
+
             for device in self.devices_sharing_workspace():
                 device.status = DeviceStatus.BUILDING
+
         elif prog.name == UnrealJobs.Sync.value:
+
             for device in self.devices_sharing_workspace():
                 device.status = DeviceStatus.SYNCING
+
+                # Reset the progress.
+                device.device_qt_handler.signal_device_sync_update.emit(device, "")
+
         elif prog.name == UnrealJobs.PackageGame.value:
             self.status = DeviceStatus.PACKAGING
 
@@ -2466,6 +2505,11 @@ class DeviceUnreal(Device):
                 LOGGER.info(f"{self.name}: Sync successful")
                 self.last_sync_filter_hash.update_value(
                     self.inflight_sync_hash)
+
+                for device in self.devices_sharing_workspace():
+                    device.device_qt_handler.signal_device_sync_update.emit(
+                        device, "Done")
+
             else:
                 LOGGER.error(f"{self.name}: Sync failed!")
                 for line in get_stdout_str().splitlines():
@@ -2480,15 +2524,15 @@ class DeviceUnreal(Device):
 
             # If you build and sync the engine, update its CL
             if CONFIG.ENGINE_SYNC_METHOD.get_value() == EngineSyncMethod.Build_Engine.value:
-                self.engine_changelist = (
-                    self.inflight_engine_cl
-                    if self.inflight_engine_cl is not None
-                    else self.engine_changelist)
-                self.inflight_engine_cl = None
+
+                if self.inflight_engine_cl:
+                    self.engine_changelist = self.inflight_engine_cl
+                    self.inflight_engine_cl = None
 
             # Update CL with the one in flight
-            self.project_changelist = self.inflight_project_cl
-            self.inflight_project_cl = None
+            if self.inflight_project_cl:
+                self.project_changelist = self.inflight_project_cl
+                self.inflight_project_cl = None
 
             self.inflight_sync_hash = None
 
@@ -2716,14 +2760,21 @@ class DeviceUnreal(Device):
         only_connected: bool = True,
         skip: Optional[DeviceUnreal] = None,
     ) -> Generator[DeviceUnreal, None, None]:
+
         def get_ws(device):
-            return (CONFIG.SOURCE_CONTROL_WORKSPACE.get_value(
-                device.name) or '').casefold()
+            return (CONFIG.SOURCE_CONTROL_WORKSPACE.get_value(device.name) or '').casefold()
+
+        def get_prj_ws(device):
+            return (CONFIG.PROJECT_WORKSPACE.get_value(device.name) or '').casefold()
 
         self_ws = get_ws(self)
+        self_prj_ws = get_prj_ws(self)
+
         def predicate(device: DeviceUnreal):
-            if self_ws:
-                return get_ws(device) == self_ws
+            if self_ws or self_prj_ws:
+                same_ws = (get_ws(device) == self_ws) if self_ws else False
+                same_prj_ws = (get_prj_ws(device) == self_prj_ws) if self_prj_ws else False
+                return same_ws or same_prj_ws
             else:
                 return device is self
 
@@ -3308,11 +3359,13 @@ class DeviceWidgetUnreal(DeviceWidget):
         self.engine_changelist_label = QtWidgets.QLabel()
         self.engine_changelist_label.setObjectName('changelist')
         self.engine_changelist_label.setToolTip(BASE_ENGINE_CL_TOOLTIP)
+        self.engine_changelist_label.setTextInteractionFlags(QtCore.Qt.TextSelectableByMouse)
         changelist_layout.addWidget(self.engine_changelist_label)
 
         self.project_changelist_label = QtWidgets.QLabel()
         self.project_changelist_label.setObjectName('changelist')
         self.project_changelist_label.setToolTip(BASE_PROJECT_CL_TOOLTIP)
+        self.project_changelist_label.setTextInteractionFlags(QtCore.Qt.TextSelectableByMouse)
         changelist_layout.addWidget(self.project_changelist_label)
 
         spacer = QtWidgets.QSpacerItem(
diff --git a/Engine/Plugins/VirtualProduction/Switchboard/Source/Switchboard/switchboard/switchboard_dialog.py b/Engine/Plugins/VirtualProduction/Switchboard/Source/Switchboard/switchboard/switchboard_dialog.py
--- a/Engine/Plugins/VirtualProduction/Switchboard/Source/Switchboard/switchboard/switchboard_dialog.py
+++ b/Engine/Plugins/VirtualProduction/Switchboard/Source/Switchboard/switchboard/switchboard_dialog.py
@@ -1325,15 +1325,23 @@ class SwitchboardDialog(QtCore.QObject):
             if (do_sync and not can_sync) or (do_build and not can_build):
                 continue
 
-            workspace = CONFIG.SOURCE_CONTROL_WORKSPACE.get_value(
-                device.name).casefold()
-            if existing_device := workspace_to_device_map.get(workspace):
-                LOGGER.info(f'Skipping {actions_str} for {device.name} '
-                            f'because {existing_device.name} is already '
-                            'queued using the same workspace')
+            engine_workspace = CONFIG.SOURCE_CONTROL_WORKSPACE.get_value(device.name).casefold()
+            project_workspace = CONFIG.PROJECT_WORKSPACE.get_value(device.name).casefold()
+
+            workspaces = [ws for ws in [engine_workspace, project_workspace] if ws]
+
+            if len(workspaces) == 0:
+                LOGGER.info(f'Skipping {actions_str} for {device.name} because it has empty workspaces')
                 continue
 
-            workspace_to_device_map[workspace] = device
+            for workspace in workspaces:
+                if existing_device := workspace_to_device_map.get(workspace):
+                    LOGGER.info(f'Skipping {actions_str} for {device.name} '
+                                f'because {existing_device.name} is already '
+                                'queued using the same workspace')
+                    continue
+
+                workspace_to_device_map[workspace] = device
 
             if do_sync:
                 device.widget.sync_button_clicked()
@@ -1814,7 +1822,18 @@ class SwitchboardDialog(QtCore.QObject):
         device = self.device_manager.device_with_hash(device_widget.device_hash)
         project_cl = None if self.project_changelist == EMPTY_SYNC_ENTRY else self.project_changelist
         engine_cl = None if self.engine_changelist == EMPTY_SYNC_ENTRY else self.engine_changelist
-        device.sync(engine_cl, project_cl)
+
+        workspace = CONFIG.SOURCE_CONTROL_WORKSPACE.get_value(device.name)
+        project_workspace = CONFIG.PROJECT_WORKSPACE.get_value(device.name)
+
+        # If the project workspace is separate and there is a target project_cl, then we need to issue separate
+        # sync commands for engine and project.
+        if project_cl and project_workspace and (project_workspace != workspace):
+            puuid = device.sync(engine_cl=engine_cl, project_cl=None)
+            device.sync(engine_cl=None, project_cl=project_cl, puuid_dependency=puuid)
+        else:
+            # Specify engine and project together because they are in the same workspace
+            device.sync(engine_cl, project_cl)
 
     @QtCore.Slot(object)
     def device_widget_build(self, device_widget):
@@ -2072,16 +2091,21 @@ class SwitchboardDialog(QtCore.QObject):
         changelists: Optional[list[int]] = None
         descriptions: Optional[list[str]] = None
 
+        client = CONFIG.SOURCE_CONTROL_WORKSPACE.get_value()
+        project_client = CONFIG.PROJECT_WORKSPACE.get_value()
+
+        if project_client:
+            client = project_client
+
         if use_ugs:
             LOGGER.info("Using UnrealGameSync to refresh project changelists.")
-            changelists = ugs_utils.latest_chagelists(Path(CONFIG.UPROJECT_PATH.get_value()), client=CONFIG.SOURCE_CONTROL_WORKSPACE.get_value())
+            changelists = ugs_utils.latest_chagelists(Path(CONFIG.UPROJECT_PATH.get_value()), client=client)
             if not changelists:
                 LOGGER.error("UnrealGameSync failed to get the project's latest changelists. Falling back to using p4 commands directly.")
 
         if not changelists:
 
             LOGGER.info("Refreshing p4 project changelists")
-            client = CONFIG.SOURCE_CONTROL_WORKSPACE.get_value()
             paths = [f'{CONFIG.P4_PROJECT_PATH.get_value()}/...']
 
             if DeviceUnreal.ugs_config:
