Switchboard: Disable Python by default when launching nDisplay.

This avoids about 10ms of reachability analysis which contribute to GC hitches.
Also added an option (default to off) to set gc.AllowIncrementalReachability to 1.

Can be re-enabled with checkbox in nDisplay section settings.

[RN] Switchboard now adds -DisablePython to nDisplay launch command lines by default.

diff --git a/Engine/Plugins/VirtualProduction/Switchboard/Source/Switchboard/switchboard/devices/ndisplay/plugin_ndisplay.py b/Engine/Plugins/VirtualProduction/Switchboard/Source/Switchboard/switchboard/devices/ndisplay/plugin_ndisplay.py
--- a/Engine/Plugins/VirtualProduction/Switchboard/Source/Switchboard/switchboard/devices/ndisplay/plugin_ndisplay.py
+++ b/Engine/Plugins/VirtualProduction/Switchboard/Source/Switchboard/switchboard/devices/ndisplay/plugin_ndisplay.py
@@ -435,6 +435,20 @@ class DevicenDisplay(DeviceUnreal):
             tool_tip='When unchecked, will add -NoLoadingScreen to the command line',
             category="UE Settings",
         ),
+        'allow_python': BoolSetting(
+            attr_name="allow_python",
+            nice_name="Python",
+            value=False,
+            tool_tip='When unchecked, will add -DisablePython to the command line',
+            category="UE Settings",
+        ),
+        'incremental_gc_reachability': BoolSetting(
+            attr_name="incremental_gc_reachability",
+            nice_name="Incremental GC reachability",
+            value=False,
+            tool_tip='When checked, will set the CVar gc.AllowIncrementalReachability to 1',
+            category="UE Settings",
+        ),
         'render_api': OptionSetting(
             attr_name="render_api",
             nice_name="Render API",
@@ -1018,6 +1032,12 @@ class DevicenDisplay(DeviceUnreal):
             if not DevicenDisplay.csettings['loading_screen'].get_value(self.name)
             else "")
 
+        # Allow Python
+        no_python = (
+            "-DisablePython"
+            if not DevicenDisplay.csettings['allow_python'].get_value(self.name)
+            else "")
+
         # MaxGPUCount (mGPU)
         max_gpu_count = DevicenDisplay.csettings["max_gpu_count"].get_value(
             self.name)
@@ -1137,6 +1157,7 @@ class DevicenDisplay(DeviceUnreal):
             f'{no_texture_streaming}',    # -notexturestreaming
             f'{no_sound}',                # -nosound
             f'{no_loading_screen}',       # -NoLoadingScreen
+            f'{no_python}',               # -DisablePython
             f'-dc_node={self.name}',      # name of this node in the nDisplay cluster
             f'Log={self.log_filename}',   # log file
             f'{ini_engine}',              # Engine ini injections
@@ -1252,6 +1273,10 @@ class DevicenDisplay(DeviceUnreal):
         if mediaprofile_gamepath:
             dp_cvars.append(self.dpcvar_for_mediaprofile(mediaprofile_gamepath))
 
+        # GC
+        if DevicenDisplay.csettings['incremental_gc_reachability'].get_value(self.name):
+            dp_cvars.append('gc.AllowIncrementalReachability=1')
+
         # Add user set dp cvars, overriding any of the forced ones.
         user_dp_cvars = self.csettings['ndisplay_dp_cvars'].get_value(self.name)
         user_dp_cvars = [cvar.strip() for cvar in user_dp_cvars if len(cvar.strip()) and len(cvar.split('=')) == 2]
