Switchboard: Default incremental GC reachability analysis to True.

#rnx

diff --git a/Engine/Plugins/VirtualProduction/Switchboard/Source/Switchboard/switchboard/devices/ndisplay/plugin_ndisplay.py b/Engine/Plugins/VirtualProduction/Switchboard/Source/Switchboard/switchboard/devices/ndisplay/plugin_ndisplay.py
--- a/Engine/Plugins/VirtualProduction/Switchboard/Source/Switchboard/switchboard/devices/ndisplay/plugin_ndisplay.py
+++ b/Engine/Plugins/VirtualProduction/Switchboard/Source/Switchboard/switchboard/devices/ndisplay/plugin_ndisplay.py
@@ -445,7 +445,7 @@ class DevicenDisplay(DeviceUnreal):
         'incremental_gc_reachability': BoolSetting(
             attr_name="incremental_gc_reachability",
             nice_name="Incremental GC reachability",
-            value=False,
+            value=True,
             tool_tip='When checked, will set the CVar gc.AllowIncrementalReachability to 1',
             category="UE Settings",
         ),
