Switchboard: Fixed regression of not showing placeholder strings in PerforcePathSetting types.

#rnx

diff --git a/Engine/Plugins/VirtualProduction/Switchboard/Source/Switchboard/switchboard/config.py b/Engine/Plugins/VirtualProduction/Switchboard/Source/Switchboard/switchboard/config.py
--- a/Engine/Plugins/VirtualProduction/Switchboard/Source/Switchboard/switchboard/config.py
+++ b/Engine/Plugins/VirtualProduction/Switchboard/Source/Switchboard/switchboard/config.py
@@ -790,7 +790,7 @@ class PerforcePathSetting(StringSetting):
         # Trim matching file paths to the parent directory (e.g. ['.uproject'])
         self.truncate_files_with_extensions: list[str] = []
 
-        super().__init__(*args, **kwargs)
+        super().__init__(*args, placeholder_text=placeholder_text, is_read_only=is_read_only, **kwargs)
 
     def _filter_value(self, value: Optional[str]) -> str:
         '''
