Switchboard: Removed deprecated use of AA_EnableHighDpiScaling.

High DPI scaling is enabled by default in PySide6.

#rnx

diff --git a/Engine/Plugins/VirtualProduction/Switchboard/Source/Switchboard/switchboard/__main__.py b/Engine/Plugins/VirtualProduction/Switchboard/Source/Switchboard/switchboard/__main__.py
--- a/Engine/Plugins/VirtualProduction/Switchboard/Source/Switchboard/switchboard/__main__.py
+++ b/Engine/Plugins/VirtualProduction/Switchboard/Source/Switchboard/switchboard/__main__.py
@@ -36,7 +36,6 @@ def launch():
         app_id = u'epicgames.virtualproduction.switchboard.0.1'
         ctypes.windll.shell32.SetCurrentProcessExplicitAppUserModelID(app_id)
 
-    QtWidgets.QApplication.setAttribute(QtCore.Qt.AA_EnableHighDpiScaling)
     app = QtWidgets.QApplication(sys.argv)
 
     # Global variables can be inited after QT has been started (they may show dialogues if there are errors)
