Switchboard: Cleanup of Setting change propagation logic and fixed a few other Setting related bugs.

Fixed some properties incorrectly updating the override widget when the base value changed. The reason was that since the setting value notification signal is connected to a callback for all override widgets, at that point the widgets could not tell if the value was being overriden or not, so they just updated their UI with the new value (even if the override value was unchanged, which is was caused the settings to disappear once the settings dialog was re-opened). Now it never notifies widgets with value changes when the values are being overridden (since widgets in that case should only show override value, not base value).

Fixed flawed logic of determining if a value is overridden or not. It was based on comparing override value with base value, but this causes the overridden values to revert their override status (and not show the override indicator) if at any moment the base value was equal to the override value (e.g. it could inadvertently stick back to the base setting). It is now based on whether there is an override value set or not, independently of whether it is equal or not to the base value. If the user wants to reset the override, the user must press the rest button next to the property.

Fixed ListSetting relying on on setting changed callback to update the override widgets. It now uses a dedicated on setting overridden signal that is now part of the Setting base class.

Fixed logic around override highlights, where specific code was needed for certain settings such as BoolSetting. Instead it is now based on whether there is a plugin setting or not, where if there isn't a plugin setting, then there is no point in the device setting to highlight since there is no base setting to override.

Also added explicit logic to avoid unecessary (and potentially risky) recursion when updating values (update value notifies widgets, which call update value again). Recursion has been avoided so far with early return when the new value is equal to the old value.

#rnx

diff --git a/Engine/Plugins/VirtualProduction/Switchboard/Source/Switchboard/switchboard/config.py b/Engine/Plugins/VirtualProduction/Switchboard/Source/Switchboard/switchboard/config.py
--- a/Engine/Plugins/VirtualProduction/Switchboard/Source/Switchboard/switchboard/config.py
+++ b/Engine/Plugins/VirtualProduction/Switchboard/Source/Switchboard/switchboard/config.py
@@ -74,9 +74,9 @@ def get_game_launch_level_path() -> str:
 
             if len(levels) == 0:
                 raise Exception('Default map is not set in DefaultEngine.ini')
-            
+
             level = levels[0]
-        except:
+        except Exception:
             LOGGER.warning("Map is not set. Please select 'Level' to run dedicated server.")
 
     return str(level)
@@ -91,6 +91,11 @@ def migrate_comma_separated_string_to_list(value) -> list[str]:
     raise NotImplementedError("Migration not handled")
 
 
+def override(method):
+    '''Decorator to indicate that a method is overriding a base class method.'''
+    return method
+
+
 class Setting(QtCore.QObject):
     '''
     A type-agnostic value container for a configuration setting.
@@ -100,7 +105,24 @@ class Setting(QtCore.QObject):
     in the UI must be represented by a derived class that creates
     the appropriate widget(s) for modifying the Setting's value.
     '''
+
+    '''
+    signal_setting_changed is emitted when the base value changes
+
+    Args:
+        old_value : old value
+        value : new value
+    '''
     signal_setting_changed = QtCore.Signal(object, object)
+
+    '''
+    signal_setting_overridden is emitted when the base value is overridden
+
+    Args:
+        device_name : Name of the device
+        value : Base value
+        override : Override value
+    '''
     signal_setting_overridden = QtCore.Signal(str, object, object)
 
     def _filter_value(self, value):
@@ -121,7 +143,8 @@ class Setting(QtCore.QObject):
         show_ui: bool = True,
         allow_reset: bool = True,
         migrate_data: Optional[Callable[[Any], None]] = None,
-        category: str = 'Misc'
+        category: str = 'Misc',
+        allow_override_style: bool = True,
     ):
         '''
         Create a new Setting object.
@@ -141,6 +164,13 @@ class Setting(QtCore.QObject):
         self.attr_name = attr_name
         self.nice_name = nice_name
         self.category = category
+        self.allow_override_style = allow_override_style
+
+        ''' updating_widgets is a flag to avoid update recursion. 
+        The Setting value or override cannot be set when True, except the first
+        time in the recursion.
+        '''
+        self.updating_widgets = False
 
         value = self._filter_value(value)
         self._original_value = self._value = value
@@ -184,40 +214,74 @@ class Setting(QtCore.QObject):
         return self
 
     def is_overridden(self, device_name: str) -> bool:
-        try:
-            return self._overrides[device_name] != self._value
-        except KeyError:
-            return False
+        return device_name in self._overrides
 
     def remove_override(self, device_name: str):
         self._overrides.pop(device_name, None)
-        self._override_widgets.pop(device_name, None)
 
     def update_value(self, new_value):
-        new_value = self._filter_value(new_value)
+        ''' Sets the base value of this setting.
+        Use override_value to override the value for a given device_name.
+        '''
 
-        if self._value == new_value:
+        # When updating widgets, we're inside a recursion, so we stop it here.
+        if self.updating_widgets:
             return
 
-        old_value = self._value
-        self._value = new_value
+        # prevent future recursion
+        self.updating_widgets = True
+
+        try:
+            new_value = self._filter_value(new_value)
+
+            if self._value == new_value:
+                return
 
-        self.signal_setting_changed.emit(old_value, self._value)
-        self._refresh_reset_base_widget()
+            old_value = self._value
+            self._value = new_value
+
+            self.signal_setting_changed.emit(old_value, self._value)
+            self._refresh_reset_base_widget()
+        finally:
+            # Re-allow value updates
+            self.updating_widgets = False
 
     def override_value(self, device_name: str, override):
-        override = self._filter_value(override)
+        ''' Overrides the base value of this setting for the given device_name
 
-        if (device_name in self._overrides and
-                self._overrides[device_name] == override):
+        Args:
+            device_name : The device name associated with the new override value.
+            override    : The value associated with the device_name.
+        '''
+        # When updating widgets, we're inside a recursion, so we stop it here.
+        if self.updating_widgets:
             return
 
-        self._overrides[device_name] = override
-        self.signal_setting_overridden.emit(device_name, self._value, override)
+        # prevent future recursion
+        self.updating_widgets = True
+
+        try:
+            override = self._filter_value(override)
 
-        self._refresh_reset_override_widget(device_name)
+            # Don't do anything if the override has the same value
+            if (device_name in self._overrides and
+                    self._overrides[device_name] == override):
+                return
+
+            self._overrides[device_name] = override
+            self.signal_setting_overridden.emit(device_name, self._value, override)
+
+            self._refresh_reset_override_widget(device_name)
+        finally:
+            # Re-allow value updates
+            self.updating_widgets = False
 
     def get_value(self, device_name: Optional[str] = None):
+        ''' Reads the value of this setting.
+
+        Args:
+            device_name : Use to read the value overridden for the given device_name
+        '''
         def get_value(self, device_name: Optional[str] = None):
             try:
                 return self._overrides[device_name]
@@ -304,13 +368,9 @@ class Setting(QtCore.QObject):
         if new_value != old_value:
             self.override_value(override_device_name, new_value)
 
-        widget = self.get_widget(override_device_name=override_device_name)
-        if self.is_overridden(override_device_name):
-            if widget:
-                sb_widgets.set_qt_property(widget, "override", True)
-        else:
-            if widget:
-                sb_widgets.set_qt_property(widget, "override", False)
+        if self.allow_override_style:
+            if (widget := self.get_widget(override_device_name)):
+                sb_widgets.set_qt_property(widget, "override", self.is_overridden(override_device_name))
 
     def _on_setting_changed(
             self, new_value,
@@ -326,6 +386,20 @@ class Setting(QtCore.QObject):
         '''
         pass
 
+    def _on_setting_overridden(
+            self,
+            new_value,
+            override_device_name: str):
+        '''
+        Callback invoked when the override value of this Setting changes.
+
+        This can be implemented in derived classes to update the appropriate
+        UI elements.
+
+        The default implementation does nothing.
+        '''
+        pass
+
     def create_ui(
             self, override_device_name: Optional[str] = None,
             form_layout: Optional[QtWidgets.QFormLayout] = None) \
@@ -354,7 +428,13 @@ class Setting(QtCore.QObject):
 
         # Give the Setting's UI an opportunity to update itself when the
         # underlying value changes.
-        self._register_on_setting_changed(top_level_widget, override_device_name)
+
+        self._register_on_setting_changed(
+            top_level_widget=top_level_widget,
+            override_device_name=override_device_name
+        )
+
+        self.signal_setting_overridden.connect(self._on_setting_overridden)
 
         if top_level_widget and form_layout:
             setting_label = QtWidgets.QLabel()
@@ -372,7 +452,12 @@ class Setting(QtCore.QObject):
     def _register_on_setting_changed(self, top_level_widget: QtWidgets.QWidget, override_device_name: str):
 
         def on_setting_changed(old_value, new_value, override_device_name=override_device_name):
-            self._on_setting_changed(new_value, override_device_name=override_device_name)
+
+            # Since this is a notification of base value changed, don't call _on_setting_changed if the value is being
+            # overridden for the given device name. Otherwise _on_setting_changed may incorrectly update the override
+            # widget with the base value.
+            if not self.is_overridden(override_device_name):
+                self._on_setting_changed(new_value, override_device_name=override_device_name)
 
         self.signal_setting_changed.connect(on_setting_changed)
 
@@ -433,13 +518,14 @@ class Setting(QtCore.QObject):
 
     def _on_press_reset_override(self, override_device_name: str):
         if override_device_name is None:
+            self.remove_override(override_device_name)
             self.update_value(self._original_value)
             self._base_reset_widget.setVisible(False)
             self._refresh_reset_override_widgets()
             return
 
         if self.is_overridden(override_device_name):
-            self.override_value(override_device_name, self._value)
+            self.remove_override(override_device_name)
 
             # Update UI
             self._on_widget_value_changed(self._value, override_device_name)
@@ -519,18 +605,11 @@ class BoolSetting(Setting):
     def _on_setting_changed(
             self, new_value: bool,
             override_device_name: Optional[str] = None):
+
         widget = self.get_widget(override_device_name=override_device_name)
         if widget:
             widget.setChecked(new_value)
 
-        if override_device_name is not None:
-            # Reset the checkbox override state. As the checkbox has only
-            # two states there is no override anymore when the base value
-            # changes.
-            if widget:
-                sb_widgets.set_qt_property(widget, "override", False)
-            self.remove_override(override_device_name)
-
 
 class IntSetting(Setting):
     '''
@@ -581,6 +660,7 @@ class IntSetting(Setting):
     def _on_setting_changed(
             self, new_value: int,
             override_device_name: Optional[str] = None):
+
         widget = self.get_widget(override_device_name=override_device_name)
         if not widget:
             return
@@ -647,6 +727,7 @@ class StringSetting(Setting):
     def _on_setting_changed(
             self, new_value: str,
             override_device_name: Optional[str] = None):
+
         widget = self.get_widget(override_device_name=override_device_name)
         if not widget:
             return
@@ -860,6 +941,7 @@ class OptionSetting(Setting):
     def _on_setting_changed(
             self, new_value,
             override_device_name: Optional[str] = None):
+
         widget = self.get_widget(override_device_name=override_device_name)
         if not widget:
             return
@@ -904,6 +986,7 @@ class MultiOptionSetting(OptionSetting):
     def _on_setting_changed(
             self, new_value,
             override_device_name: Optional[str] = None):
+
         widget = self.get_widget(override_device_name=override_device_name)
         if not widget:
             return
@@ -990,7 +1073,7 @@ class ListRow(QtWidgets.QWidget):
 class ListSetting(Setting):
     '''
     A setting which has add and clear buttons. Each item will be displayed in a new line.
-    Functions like array properties in Unreal Editor.
+    Operates similarly to array properties in Unreal Editor.
 
     Subclasses are responsible for generating widgets for the array contents, see e.g. ArrayStringSetting.
     '''
@@ -1168,10 +1251,43 @@ class ListSetting(Setting):
         )
         return row, default_value
 
-    def _on_setting_changed(self, new_value: list, override_device_name: Optional[str] = None):
+    @override
+    def _on_setting_changed(
+        self,
+        new_value: list,
+        override_device_name: Optional[str] = None
+    ) -> None:
+
+        self.update_list_ui(new_value=new_value, override_device_name=override_device_name)
+
+    @override
+    def _on_setting_overridden(
+        self,
+        new_value,
+        override_device_name: str
+    ) -> None:
+
+        self.update_list_ui(new_value=new_value, override_device_name=override_device_name)
+
+    def update_list_ui(
+        self,
+        new_value,
+        override_device_name: str
+    ) -> None:
+        '''
+        Update the UI to reflect the given list value.
+
+        Args:
+            new_value (list):
+                The new list of values that should be represented in the UI.
+
+            override_device_name (Optional[str]):
+                The name of the device to which the override applies.
+                If provided, updates the UI for the specific device.
+        '''
+
         container_layout = self.element_layouts[override_device_name]
 
-        container_len = container_layout.count()
         new_len = len(new_value)
         new_list_is_bigger = container_layout.count() < new_len
         new_list_is_smaller = container_layout.count() > new_len
@@ -1203,7 +1319,7 @@ class StringListSetting(ListSetting):
                 self._on_editor_widget_changed(line_edit, override_device_name, index)
         )
         return line_edit, ""
-    
+
     def _on_editor_widget_changed(self, line_edit, override_device_name, index):
         current_list = self.get_value(override_device_name)
         # List needs to a different object otherwise _on_widget_value_changed will think the value has not changed
@@ -1365,7 +1481,7 @@ class LoggingModel(QtGui.QStandardItemModel):
 
         self.beginResetModel()
         count_before = self.rowCount()
-        
+
         root_item = self.invisibleRootItem()
         for category, verbosity_level in value.items():
             if (category not in self._categories and
@@ -1375,9 +1491,9 @@ class LoggingModel(QtGui.QStandardItemModel):
             root_item.appendRow(
                 [QtGui.QStandardItem(category),
                     QtGui.QStandardItem(verbosity_level)])
-            
-        # This triggers the rowsRemoved event. 
-        # Remove rows after so external observers get the right result from category_verbosities.  
+
+        # This triggers the rowsRemoved event.
+        # Remove rows after so external observers get the right result from category_verbosities.
         self.removeRows(0, count_before)
         self.endResetModel()
 
@@ -1433,11 +1549,10 @@ class LoggingVerbosityItemDelegate(QtWidgets.QStyledItemDelegate):
         )
         # For multi-selection we want to also be modified even when index has not changed
         editor.activated.connect(
-            lambda combo_index, 
-                editor=editor:
+            lambda combo_index, editor=editor:
                 self.on_current_index_changed(combo_index, editor)
         )
-        
+
         return editor
 
     def setEditorData(
@@ -1485,26 +1600,26 @@ class LoggingSettingVerbosityView(QtWidgets.QTreeView):
         self.setItemDelegateForColumn(
             LoggingModel.VERBOSITY_COLUMN,
             LoggingVerbosityItemDelegate(logging_model.verbosity_levels, self))
-        
+
         self._open_persistent_editors()
         logging_model.modelAboutToBeReset.connect(self._pre_change_model)
         logging_model.modelReset.connect(self._post_change_model)
 
         self.setSelectionBehavior(QtWidgets.QTreeView.SelectRows)
         self.setSelectionMode(QtWidgets.QTreeView.ExtendedSelection)
-        
+
     def _pre_change_model(self):
         self.scroll_height = self.verticalScrollBar().value()
-        
+
     def _post_change_model(self):
         self.verticalScrollBar().setSliderPosition(self.scroll_height)
         self._open_persistent_editors()
-        
+
     def _open_persistent_editors(self):
         for row in range(self.model().rowCount()):
             verbosity_level_index = self.model().index(row, LoggingModel.VERBOSITY_COLUMN)
             self.openPersistentEditor(verbosity_level_index)
-        
+
     def selected_categories(self) -> list[str]:
         model = self.model()
         selected_categories = []
@@ -1514,9 +1629,9 @@ class LoggingSettingVerbosityView(QtWidgets.QTreeView):
                 category_index = model.index(row, 0)
                 category = model.data(category_index, QtCore.Qt.DisplayRole)
                 selected_categories.append(category)
-                
+
         return selected_categories
-        
+
     def update_category_verbosities(self, categories: list[str], new_verbosity: str):
         model = self.model()
         category_verbosities = model.category_verbosities
@@ -1718,7 +1833,7 @@ class LoggingSetting(Setting):
 
         view.selectionModel().selectionChanged.connect(
             on_view_selectionChanged)
-        
+
         def on_logging_model_modified(override_device_name=None):
             category_verbosities = model.category_verbosities
             self._on_widget_value_changed(
@@ -1742,13 +1857,13 @@ class LoggingSetting(Setting):
             override_device_name=override_device_name:
                 on_logging_model_modified(
                     override_device_name=override_device_name))
-        
-        
+
         return edit_layout
 
     def _on_setting_changed(
             self, new_value,
             override_device_name: Optional[str] = None):
+
         widget = self.get_widget(override_device_name=override_device_name)
         if not widget:
             return
@@ -1796,6 +1911,7 @@ class AddressSetting(OptionSetting):
     def generate_possible_addresses(self):
         return set[str]()
 
+
 def is_valid_unicast_address(address):
     # 169.254.x.x addresses can be generated by the OS when a failure to negotiate with DHCP.
     if address.startswith("169.254"):
@@ -1822,23 +1938,29 @@ def is_valid_unicast_address(address):
 
     return True
 
+
 def is_local_classC_ip(address):
     if address.startswith("192."):
         return True
     return False
 
+
 def is_local_classA_ip(address):
     if address.startswith("10."):
         return True
     return False
 
+
 class LocalAddressSetting(AddressSetting):
     def check_and_fix_bad_ip(self):
+
         orig_value = self.get_value()
+
         if is_valid_unicast_address(orig_value):
             return
 
         addresses = self.generate_possible_addresses()
+
         def iterate_over_addresses():
             classC = filter(lambda x: is_local_classC_ip(x), addresses)
             classA = filter(lambda x: is_local_classA_ip(x), addresses)
@@ -2938,7 +3060,7 @@ class UserSettings(object):
             data = {}
         except ValueError:
             LOGGER.error(f'{USER_SETTINGS_FILE_NAME} has invalid JSON format. ')
-            
+
             answer = QtWidgets.QMessageBox.question(
                 None,
                 'Invalid User Settings',
@@ -2950,7 +3072,7 @@ class UserSettings(object):
                 if os.path.exists(USER_SETTINGS_BACKUP_FILE_PATH):
                     os.remove(USER_SETTINGS_BACKUP_FILE_PATH)
                 shutil.copy(USER_SETTINGS_FILE_PATH, USER_SETTINGS_BACKUP_FILE_PATH)
-                
+
             data = {}
 
         self.CONFIG = data.get('config')
@@ -2976,7 +3098,6 @@ class UserSettings(object):
                 "Transport path",
                 data.get('transport_path', '')
             )
-        
 
         # UI Settings
         self.MUSERVER_SESSION_NAME = data.get(
diff --git a/Engine/Plugins/VirtualProduction/Switchboard/Source/Switchboard/switchboard/devices/device_base.py b/Engine/Plugins/VirtualProduction/Switchboard/Source/Switchboard/switchboard/devices/device_base.py
--- a/Engine/Plugins/VirtualProduction/Switchboard/Source/Switchboard/switchboard/devices/device_base.py
+++ b/Engine/Plugins/VirtualProduction/Switchboard/Source/Switchboard/switchboard/devices/device_base.py
@@ -16,7 +16,7 @@ import pythonosc.udp_client
 
 from switchboard import config_osc as osc
 from switchboard import recording
-from switchboard.config import CONFIG, BoolSetting, AddressSetting
+from switchboard.config import CONFIG, BoolSetting, AddressSetting, Setting
 from switchboard.switchboard_logging import LOGGER
 
 if TYPE_CHECKING:
@@ -130,6 +130,8 @@ class Device(QtCore.QObject):
             lambda _, new_address, widget=self.widget:
                 widget.on_address_changed(new_address))
 
+        self.disable_override_style_on_device_only_settings()
+
         # Let the CONFIG class know what settings/properties this device has
         # so it can save the config file on changes.
         CONFIG.register_device_settings(
@@ -146,6 +148,16 @@ class Device(QtCore.QObject):
         if self.address != new_address:
             self.address = new_address
 
+    def disable_override_style_on_device_only_settings(self):
+        ''' If a given setting does not exist in the plugin settings, then there is
+        no point in supporting override highlights. If we don't do this, we may get
+        override highlights on device instance properties that don't have a plugin
+        value.
+        '''
+        for setting in self.device_settings():
+            if setting.attr_name not in self.plugin_settings():
+                setting.allow_override_style = False
+
     @property
     def device_type(self):
         return self.__class__.__name__.split("Device", 1)[1]
@@ -170,11 +182,11 @@ class Device(QtCore.QObject):
             csetting.reset()
 
     @classmethod
-    def plugin_settings(cls):
+    def plugin_settings(cls) -> list[Setting]:
         # settings that are shared by all devices of a plugin
         return list(cls.csettings.values())
 
-    def device_settings(self):
+    def device_settings(self) -> list[Setting]:
         # settings that are specific to an instance of the device
         return [self.setting_address]
 
diff --git a/Engine/Plugins/VirtualProduction/Switchboard/Source/Switchboard/switchboard/devices/unreal/plugin_unreal.py b/Engine/Plugins/VirtualProduction/Switchboard/Source/Switchboard/switchboard/devices/unreal/plugin_unreal.py
--- a/Engine/Plugins/VirtualProduction/Switchboard/Source/Switchboard/switchboard/devices/unreal/plugin_unreal.py
+++ b/Engine/Plugins/VirtualProduction/Switchboard/Source/Switchboard/switchboard/devices/unreal/plugin_unreal.py
@@ -30,9 +30,10 @@ import switchboard.config_osc as osc
 import switchboard.switchboard_dialog as sb_dialog
 import switchboard.switchboard_utils as sb_utils
 import switchboard.switchboard_widgets as sb_widgets
-from switchboard.config import (CONFIG, DEFAULT_MAP_TEXT, ENABLE_UGS_SUPPORT, SETTINGS, BoolSetting, DirectoryPathSetting,
-                                EngineSyncMethod, FilePathSetting, IntSetting, MultiOptionSetting, OptionSetting, Setting,
-                                StringListSetting, StringSetting, migrate_comma_separated_string_to_list)
+from switchboard.config import (
+    CONFIG, DEFAULT_MAP_TEXT, ENABLE_UGS_SUPPORT, SETTINGS, BoolSetting, DirectoryPathSetting,
+    EngineSyncMethod, FilePathSetting, IntSetting, MultiOptionSetting, OptionSetting, Setting,
+    StringListSetting, StringSetting, migrate_comma_separated_string_to_list)
 from switchboard.devices.device_base import Device, DeviceStatus, PluginHeaderWidgets
 from switchboard.devices.device_widget_base import DeviceAutoJoinMUServerUI, DeviceWidget
 from switchboard.devices.unreal.uassetparser import UassetParser
@@ -447,7 +448,9 @@ class LiveLinkPresetSetting(Setting):
         ''' Return the valid class names of this type of asset '''
         return DeviceUnreal.LIVELINKPRESET_CLASS_NAMES
 
-    def _on_setting_changed(self, new_value: str, override_device_name: Optional[str] = None):
+    def _on_setting_changed(
+            self, new_value: str,
+            override_device_name: Optional[str] = None):
 
         combo = self.get_widget(override_device_name=override_device_name)
 
@@ -3727,5 +3730,5 @@ class DeviceWidgetUnreal(DeviceWidget):
         if self.connect_button.isChecked():
             fill_ddc_menu = cmenu.addMenu("Fill DDC (Prepare Shaders)")
 
-            current_level_action = fill_ddc_menu.addAction("Current Level", lambda: self.signal_device_widget_fill_ddc.emit(self, True))
-            all_levels_action = fill_ddc_menu.addAction("All Levels", lambda: self.signal_device_widget_fill_ddc.emit(self, False))
+            fill_ddc_menu.addAction("Current Level", lambda: self.signal_device_widget_fill_ddc.emit(self, True))
+            fill_ddc_menu.addAction("All Levels", lambda: self.signal_device_widget_fill_ddc.emit(self, False))
