Switchboard: Make settings dialog appear fast with even with tens of devices.

Tested with 50 devices, without this change it was taking about 12 seconds to load the settings dialog, now it takes less than 1 second to appear.

The dialog is now shown almost right away, with all the setting categories, and the settings of each category are added incremetally per main loop, where it stops adding when the time spent per loop exceeds 16 milliseconds. To do this it keeps a pointer to the layout widget and adds rows to it.

To make search still work, it detects when new rows are added by comparing with its internal row cache that it uses to remove/re-add to the form layouts according to the search results. While the incremental add is happening, the search is renewed every time it sees that 50 settings have been added.

#rn Switchboard: Improved settings dialog responsiveness when displaying properties for tens of devices.

diff --git a/Engine/Plugins/VirtualProduction/Switchboard/Source/Switchboard/switchboard/config.py b/Engine/Plugins/VirtualProduction/Switchboard/Source/Switchboard/switchboard/config.py
--- a/Engine/Plugins/VirtualProduction/Switchboard/Source/Switchboard/switchboard/config.py
+++ b/Engine/Plugins/VirtualProduction/Switchboard/Source/Switchboard/switchboard/config.py
@@ -13,7 +13,7 @@ import socket
 import sys
 import threading
 import time
-from typing import Any, Callable, Optional, Tuple, Type, Union
+from typing import Any, Callable, Optional, Tuple, Type, Union, List
 
 from PySide6 import QtCore
 from PySide6 import QtGui
@@ -96,6 +96,132 @@ def override(method):
     return method
 
 
+def get_parent_dialog_or_window(widget: QtWidgets.QWidget) -> QtWidgets.QWidget:
+    '''
+    Gets the parent dialog or main window of the given widget.
+
+    Args:
+        widget : The widget for which to find the parent dialog or main window.
+    '''
+    if widget is None:
+        return None
+
+    parent = widget.parentWidget()
+
+    while parent and not isinstance(parent, QtWidgets.QDialog) and not isinstance(parent, QtWidgets.QMainWindow):
+        parent = parent.parentWidget()
+
+    return parent
+
+
+class BatchItemToCreateUI:
+    ''' create_ui batch item with all the Setting to be generated and added to the destination form layout '''
+
+    def __init__(
+            self,
+            setting: 'Setting',
+            layout: QtWidgets.QFormLayout,
+            device_name: Optional[str] = None):  
+
+        self.setting = setting  # the setting we're creating the top widget for
+        self.layout = layout  # the layout where it should be added
+        self.device_name = device_name  # The device override name this setting applies to
+
+
+class SettingsBatchManagerToCreateUI(QtCore.QObject):
+    '''
+    Manages the batch processing of settings dialog widgets, adding them to the UI in manageable
+    chunks to improve responsiveness.
+    '''
+
+    # Our singleton instance
+    _instance = None
+
+    @classmethod
+    def get_instance(cls):
+        if cls._instance is None:
+            cls._instance = cls()
+        return cls._instance
+
+    '''
+    Signal emitted when a batch of settings is completed.
+
+    Args:
+        int: Number of settings added during the batch. Receiver can use this to react to progress.
+        bool: True if there is more pending work, False otherwise.
+    '''
+    signal_batch_completed = QtCore.Signal(int, bool)
+
+    def __init__(self, max_milliseconds_per_batch=1000//60):
+
+        # Singleton pattern.
+
+        if self.__class__._instance is not None:
+            raise Exception("This class is a singleton!")
+
+        self.__class__._instance = self
+
+        super().__init__()
+
+        self.pending_settings: List['Setting'] = []
+        self.max_milliseconds_per_batch = max_milliseconds_per_batch
+        self.batch_timer = QtCore.QTimer()
+        self.batch_timer.timeout.connect(self.process_next_batch)
+
+    def add_setting(
+            self,
+            setting: 'Setting',
+            form_layout: QtWidgets.QFormLayout,
+            override_device_name: Optional[str] = None):
+        '''
+        Adds a setting to the batch manager, to be processed later in a batch.
+
+        Args:
+            setting : The setting instance to add.
+            form_layout : The layout to add the setting widget to.
+            override_device_name : The device name to override. Can be None.
+        '''
+
+        self.pending_settings.append(BatchItemToCreateUI(setting, form_layout, override_device_name))
+
+        # Start processing batches if the timer is not already active
+        if not self.batch_timer.isActive():
+            self.batch_timer.start(0)  # Start the timer to process on next event loop
+
+    def process_next_batch(self):
+        '''
+        Processes the next batch of settings, replacing placeholders with the actual widgets.
+        '''
+        start_time = QtCore.QElapsedTimer()
+        start_time.start()
+
+        settings_done_count = 0
+
+        while self.pending_settings:
+            batchitem = self.pending_settings.pop(0)
+            settings_done_count += 1
+
+            # Skip the work if the parent dialog or window has been destroyed or is no longer visible
+            owner = get_parent_dialog_or_window(batchitem.layout)
+            if not owner or not owner.isVisible():
+                continue
+
+            batchitem.setting.add_setting_to_layout(batchitem)
+
+            # Check if the maximum allowed time per batch has been exceeded
+            if start_time.elapsed() > self.max_milliseconds_per_batch:
+                break
+
+        if not self.pending_settings:
+            self.batch_timer.stop()  # Stop the timer when there are no more pending settings
+
+        # Notify our observers of the progress
+        self.signal_batch_completed.emit(settings_done_count, self.has_pending_work())
+
+    def has_pending_work(self):
+        return len(self.pending_settings) > 0
+
+
 class Setting(QtCore.QObject):
     '''
     A type-agnostic value container for a configuration setting.
@@ -401,53 +527,61 @@ class Setting(QtCore.QObject):
         pass
 
     def create_ui(
-            self, override_device_name: Optional[str] = None,
+            self,
+            override_device_name: Optional[str] = None,
             form_layout: Optional[QtWidgets.QFormLayout] = None) \
-            -> Union[QtWidgets.QWidget, QtWidgets.QLayout, None]:
+            -> None:
         '''
-        Create the UI for this Setting.
-
-        A label will be created for the Setting and a row for it will be
-        added to the provided form layout along with the widget/layout
-        returned by _create_widgets().
-
-        Ideally, derived classes will not need to override this function
-        and can instead just implement _create_widgets().
+        Create the UI for this Setting. The widget will be added incrementally by the batch manager.
         '''
+
+        # Honor request to not be shown in the UI
         if not self.show_ui:
-            # It's ok to use the Setting base class directly for
-            # settings that do not display in the UI.
             return None
 
-        top_level_widget = self._create_widgets(
-            override_device_name=override_device_name)
+        if form_layout:
+            # Schedule the widget to be added incrementally
+            self.get_create_ui_batch_manager().add_setting(self, form_layout, override_device_name)
 
-        widget = self.get_widget(override_device_name=override_device_name)
-        if widget and self.is_overridden(override_device_name):
-            sb_widgets.set_qt_property(widget, "override", True)
+        return None
 
-        # Give the Setting's UI an opportunity to update itself when the
-        # underlying value changes.
+    def add_setting_to_layout(self, item: BatchItemToCreateUI):
+        '''
+        Adds a setting to the provided layout.
 
-        self._register_on_setting_changed(
-            top_level_widget=top_level_widget,
-            override_device_name=override_device_name
-        )
+        Args:
+            item (BatchItemToCreateUI):
+                The batch item containing the setting information, layout, and an optional 
+                device name.
+        '''
+
+        # Create the widget for this setting.
+        top_level_widget = self._create_widgets(item.device_name)
+
+        widget = self.get_widget(item.device_name)
+        if widget and self.is_overridden(item.device_name):
+            sb_widgets.set_qt_property(widget, "override", True)
 
-        self.signal_setting_overridden.connect(self._on_setting_overridden)
+        # Register the Setting's change listeners
+        self._register_on_setting_changed(top_level_widget, item.device_name)
 
-        if top_level_widget and form_layout:
+        if top_level_widget and item.layout:
+
+            # Create the row label.
             setting_label = QtWidgets.QLabel()
             setting_label.setText(self.nice_name)
             if self.tool_tip:
                 setting_label.setToolTip(self.tool_tip)
 
-            form_layout.addRow(
-                setting_label,
-                self._decorate_with_reset_widget(override_device_name, top_level_widget)
-            )
+            # Create the field widget
+            field_widget = self._decorate_with_reset_widget(item.device_name, top_level_widget)
 
-        return top_level_widget
+            # Add the row with the label and widget to the layout.
+            item.layout.addRow(setting_label, field_widget)
+
+    @staticmethod
+    def get_create_ui_batch_manager():
+        return SettingsBatchManagerToCreateUI.get_instance()
 
     def _register_on_setting_changed(self, top_level_widget: QtWidgets.QWidget, override_device_name: str):
 
diff --git a/Engine/Plugins/VirtualProduction/Switchboard/Source/Switchboard/switchboard/settings_dialog.py b/Engine/Plugins/VirtualProduction/Switchboard/Source/Switchboard/switchboard/settings_dialog.py
--- a/Engine/Plugins/VirtualProduction/Switchboard/Source/Switchboard/switchboard/settings_dialog.py
+++ b/Engine/Plugins/VirtualProduction/Switchboard/Source/Switchboard/switchboard/settings_dialog.py
@@ -2,7 +2,7 @@
 
 import os
 import pathlib
-from typing import Dict, Optional
+from typing import Dict, Optional, Callable
 
 from PySide6 import QtCore
 from PySide6 import QtGui
@@ -28,6 +28,64 @@ COLLAPSE_ALL_TEXT = "Collapse all"
 EXPAND_ALL_TEXT = "Expand all"
 
 
+class CreateUISearchUpdater:
+    ''' Updates the search based on the progress of the settings batch manager
+    that does the work enqueued via its create_ui method.
+    '''
+
+    MIN_SETTINGS_ADDED_PER_UPDATE = 50
+
+    def __init__(
+            self,
+            batch_completed_signal: QtCore.Signal,
+            search_settings_fn: Callable):
+        '''
+        Initializes the search updater that subscribes to the batch manager's signal.
+
+        Args:
+            batch_completed_signal (int, bool): 
+                Progress signal emitted by the batch manager.
+                The first argument (int) is how many settings were added in this batch.
+                The second argument (bool) is True if there is still work pending.
+            search_settings_fn : The function to call to update the search.
+        '''
+        if batch_completed_signal is None:
+            raise ValueError("batch_manager cannot be None")
+
+        if search_settings_fn is None:
+            raise ValueError("search_settings_fn cannot be None")
+
+        self.search_settings_fn = search_settings_fn
+        self.batch_completed_signal = batch_completed_signal
+        self.settings_added_since_last_search_update = 0
+
+        # Connect the batch manager's signal to the search update function
+        self.batch_completed_signal.connect(self.update_search_results)
+
+    def update_search_results(self, settings_added: int, has_pending_work: bool):
+        '''
+        Updates the search results incrementally as new settings are added.
+        Signature matches batch_completed_signal.
+
+        Args:
+            settings_added : Number of settings added during the batch.
+            has_pending_work : True if there is more pending work, False otherwise.
+        '''
+
+        self.settings_added_since_last_search_update += settings_added
+
+        if not has_pending_work or self.settings_added_since_last_search_update > self.MIN_SETTINGS_ADDED_PER_UPDATE:
+            # Reset the count
+            self.settings_added_since_last_search_update = 0
+
+            # Trigger the search
+            self.search_settings_fn()  # None is not allowed by our constructor.
+
+    def __del__(self):
+        if self.batch_completed_signal and self.update_search_results:
+            self.batch_completed_signal.disconnect(self.update_search_results)    
+
+
 class SettingsDialog(QtCore.QObject):
     def __init__(self, settings: UserSettings, config: Config):
         super().__init__()
@@ -55,8 +113,18 @@ class SettingsDialog(QtCore.QObject):
         self._create_tab_widget(dialog_layout)
 
         self.settings_search = SettingsSearch([self.ui.all_settings_scroll_area, self.ui.general_settings_scroll_area])
-        self.ui.searchBar.textChanged.connect(
-            self._on_search_text_edited)
+        self.ui.searchBar.textChanged.connect(self._on_search_text_edited)
+
+        # Because the batch manager will be incrementally adding setting widgets, we need to keep updating
+        # the search results until it is done.
+
+        def update_search_fn():
+            '''Helper function to only conduct the search when needed'''
+            if len(self.ui.searchBar.text()) > 0:
+                self._search_settings_again()
+
+        batch_completed_signal = Setting.get_create_ui_batch_manager().signal_batch_completed
+        self.search_updater = CreateUISearchUpdater(batch_completed_signal, update_search_fn)
 
     def select_all_tab(self):
         self._on_tab_changed(0)
@@ -127,6 +195,9 @@ class SettingsDialog(QtCore.QObject):
                 # originally specified path when queried via config_path().
                 self._changed_config_path = None
 
+        # Null the search updater to break any potential circular references.
+        self.search_updater = None
+
     def _search_settings_again(self):
         self._on_search_text_edited(self.ui.searchBar.text())
 
diff --git a/Engine/Plugins/VirtualProduction/Switchboard/Source/Switchboard/switchboard/settings_search.py b/Engine/Plugins/VirtualProduction/Switchboard/Source/Switchboard/switchboard/settings_search.py
--- a/Engine/Plugins/VirtualProduction/Switchboard/Source/Switchboard/switchboard/settings_search.py
+++ b/Engine/Plugins/VirtualProduction/Switchboard/Source/Switchboard/switchboard/settings_search.py
@@ -1,7 +1,9 @@
 # Copyright Epic Games, Inc. All Rights Reserved.
 
-from typing import List, Dict
 from PySide6 import QtWidgets
+from typing import List, Dict
+
+from switchboard.switchboard_logging import LOGGER
 from switchboard.switchboard_widgets import CollapsibleGroupBox
 
 
@@ -9,6 +11,14 @@ def as_widget_or_layout(row_item):
     return row_item.widget() if row_item.widget() else row_item.layout()
 
 
+class FormLayoutRow:
+    ''' Item with the information about a form layout row. Managed by FormLayoutRowVisibility '''
+    def __init__(self, label: QtWidgets.QWidget, field: QtWidgets.QWidget, visible: bool):
+        self.label = label
+        self.field = field
+        self.visible = visible
+
+
 class FormLayoutRowVisibility:
     """
     Handles the showing and hiding of rows in QFormLayout
@@ -21,54 +31,99 @@ class FormLayoutRowVisibility:
     def __init__(self, form: QtWidgets.QFormLayout):
         self.form = form
         # Save the original widget order so we can preserve it when adding hidden widgets back
-        self.rows = []
+
+        # Keep track of visible rows. Will help us quickly determine if there is a count mismatch
+        # which would indicate that the cache should be updated.
+        self.visible_rows_count = 0
+
+        self.rows: list[FormLayoutRow] = []  # Only modify with modifier functions, to keep track of statistics.
+
         for i in range(form.rowCount()):
             label = form.itemAt(i, QtWidgets.QFormLayout.LabelRole)
             field = form.itemAt(i, QtWidgets.QFormLayout.FieldRole)
             if label and field:
-                pair = (
-                    as_widget_or_layout(label),
-                    as_widget_or_layout(field),
-                    True
-                )
-                self.rows.append(pair)
+                row = FormLayoutRow(as_widget_or_layout(label), as_widget_or_layout(field), True)
+                self._append_to_rows(row)
+
+    def update_rows_cache(self):
+        '''
+        Updates the internal row cache to include new rows added to the form.
+        We expect SettingsBatchManagerToCreateUI to incrementally add rows to form layouts.
+        We do not expect or support rows being removed from the form layout.
+        '''
+        for rowIdx in range(self.form.rowCount()):
+            label = as_widget_or_layout(self.form.itemAt(rowIdx, QtWidgets.QFormLayout.LabelRole))
+            field = as_widget_or_layout(self.form.itemAt(rowIdx, QtWidgets.QFormLayout.FieldRole))
+
+            found = False
+
+            for cached_row in self.rows:
+                if cached_row.label is label and cached_row.field is field:
+                    found = True
+                    break
+
+            # If we didn't find a corresponding entry, then we should add it to the internal row cache.
+            if not found:
+                self._append_to_rows(FormLayoutRow(label, field, True))
+
+        # sanity check that the cache is in sync.
+        if self.has_form_added_rows_externally():
+            LOGGER.error("FormLayoutRowVisibility's row cache was updated but differences were still detected.")
+
+    def has_form_added_rows_externally(self):
+        ''' Since visible rows should still be in the layout, the number of visible rows 
+        should match the number of rows in the form layout unless more rows were added externally'''
+        return self.form.rowCount() != self.visible_rows_count
+
+    def _append_to_rows(self, row: FormLayoutRow) -> None:
+        ''' Appends to the internal row cache, keeping track of visible row count.'''
+        self.rows.append(row)
+        if row.visible:
+            self.visible_rows_count += 1
+
+    def _replace_row(self, index: int, row: FormLayoutRow) -> None:
+        ''' Replaces an item of internal row cache, adjusting the visible row count.'''
+        old_visible = self.rows[index].visible
+        self.rows[index] = row
+
+        # update visible count
+        if not old_visible and row.visible:
+            self.visible_rows_count += 1
+        elif old_visible and not row.visible:
+            self.visible_rows_count -= 1
 
     def row_count(self):
         return len(self.rows)
 
     def get_label_item(self, index: int):
-        return self.rows[index][0]
+        return self.rows[index].label
 
     def set_row_visible(self, index: int, visible: bool):
-        (label, field, is_visible) = self.rows[index]
-        if visible and not is_visible:
+        row = self.rows[index]
+        if visible and not row.visible:
             self._show_row(index)
-        elif not visible and is_visible:
+        elif not visible and row.visible:
             self._hide_row(index)
 
     def _show_row(self, index: int):
-        (label, field, is_visible) = self.rows[index]
+        row = self.rows[index]
         index_to_use = 0 if index == 0 else self._find_index_to_insert_row(index)
-        self.form.insertRow(index_to_use, label, field)
-        self.rows[index] = (
-            label,
-            field,
-            True
-        )
+        self.form.insertRow(index_to_use, row.label, row.field)
+        self._replace_row(index, FormLayoutRow(row.label, row.field, True))
 
         # Widgets were hidden when the row was hidden
-        self._set_visible_recursive(label, True)
-        self._set_visible_recursive(field, True)
+        self._set_visible_recursive(row.label, True)
+        self._set_visible_recursive(row.field, True)
 
     def _hide_row(self, index: int):
-        (label, field, is_visible) = self.rows[index]
+        row = self.rows[index]
         # "index" is the original index of the row before we any rows were hidden.
         # Since we may have hidden certain rows already, we must search for the the row that corresponds "index"
         for i in range(self.form.rowCount()):
             label_item = self.form.itemAt(i, QtWidgets.QFormLayout.LabelRole)
-            if label == as_widget_or_layout(label_item):
+            if row.label == as_widget_or_layout(label_item):
                 self._take_row(i)
-                self.rows[index] = (label, field, False)
+                self._replace_row(index, FormLayoutRow(row.label, row.field, False))
                 # Found the matching row
                 break
 
@@ -90,12 +145,11 @@ class FormLayoutRowVisibility:
         self.form.update()
 
     def _find_index_to_insert_row(self, index: int):
-        (parent_label, _, _) = self.rows[index - 1]
+        row = self.rows[index - 1]
         for actual_form_index in range(self.form.rowCount()):
             actual_item = self.form.itemAt(actual_form_index, QtWidgets.QFormLayout.LabelRole)
-            for saved_form_index in range(index):
-                if as_widget_or_layout(actual_item) == parent_label:
-                    return actual_form_index + 1
+            if as_widget_or_layout(actual_item) is row.label:
+                return actual_form_index + 1
 
         return 0
 
@@ -202,6 +256,10 @@ class SettingsSearch:
         is_match_on_any_child = False
         form_handler = self.form_layout_handlers.setdefault(layout, FormLayoutRowVisibility(layout))
 
+        # Update the rows cache of the form handler if SettingsBatchManagerToCreateUI as added more rows.
+        if form_handler.has_form_added_rows_externally():
+            form_handler.update_rows_cache()
+
         for idx in range(form_handler.row_count()):
             label = form_handler.get_label_item(idx)
             if isinstance(label, QtWidgets.QWidget):
