Switchboard: Allow user to specify each DPCVar key/value pair separated by a space, as one would in the Unreal command console.

diff --git a/Engine/Plugins/VirtualProduction/Switchboard/Source/Switchboard/switchboard/devices/ndisplay/plugin_ndisplay.py b/Engine/Plugins/VirtualProduction/Switchboard/Source/Switchboard/switchboard/devices/ndisplay/plugin_ndisplay.py
--- a/Engine/Plugins/VirtualProduction/Switchboard/Source/Switchboard/switchboard/devices/ndisplay/plugin_ndisplay.py
+++ b/Engine/Plugins/VirtualProduction/Switchboard/Source/Switchboard/switchboard/devices/ndisplay/plugin_ndisplay.py
@@ -1278,8 +1278,7 @@ class DevicenDisplay(DeviceUnreal):
             dp_cvars.append('gc.AllowIncrementalReachability=1')
 
         # Add user set dp cvars, overriding any of the forced ones.
-        user_dp_cvars = self.csettings['ndisplay_dp_cvars'].get_value(self.name)
-        user_dp_cvars = [cvar.strip() for cvar in user_dp_cvars if len(cvar.strip()) and len(cvar.split('=')) == 2]
+        user_dp_cvars = self.conform_dpcvars(self.csettings['ndisplay_dp_cvars'].get_value(self.name))
         dp_cvars = self.add_or_override_cvars(dp_cvars, user_dp_cvars)
 
         # add accumulated dpcvars to args
diff --git a/Engine/Plugins/VirtualProduction/Switchboard/Source/Switchboard/switchboard/devices/unreal/plugin_unreal.py b/Engine/Plugins/VirtualProduction/Switchboard/Source/Switchboard/switchboard/devices/unreal/plugin_unreal.py
--- a/Engine/Plugins/VirtualProduction/Switchboard/Source/Switchboard/switchboard/devices/unreal/plugin_unreal.py
+++ b/Engine/Plugins/VirtualProduction/Switchboard/Source/Switchboard/switchboard/devices/unreal/plugin_unreal.py
@@ -2120,6 +2120,45 @@ class DeviceUnreal(Device):
 
         return [f'{cvar[0]}={cvar[1]}' for cvar in cvars_map.values()]
 
+    def conform_dpcvars(self, cvars: list[str]) -> list[str]:
+        ''' Conforms input cvar to expected format for command line (key=value)
+        It will allow the user to separate key/value using space(s), like one would in Unreal's command input field.
+        '''
+        ocvars = []
+
+        disallowed_chars = [
+            ',',  # Comma is used as a dpcvar delimiter.
+        ]
+
+        for cvar in cvars:
+            # remove extra spaces
+            cvar = cvar.strip()
+
+            # Split by '=' (allow surround by spaces) or 1+ spaces
+            match = re.split(r'\s*=\s*|\s{1,}', cvar, maxsplit=1)
+
+            if len(match) != 2:
+                LOGGER.warning(f"Ignored cvar '{cvar}' because it could not be parsed")
+                continue
+
+            key, value = match
+
+            # Check for empty key or value
+            if not key or not value:
+                LOGGER.warning(f"Ignored cvar '{cvar}' because it has an empty key or value")
+                continue
+
+            cvar = f"{key}={value}"
+
+            # Check for disallowed characters
+            if (any(char in cvar for char in disallowed_chars)):
+                LOGGER.warning(f"Ignored cvar '{cvar}' because it contained one or more invalid characters ({disallowed_chars})")
+                continue
+
+            ocvars.append(cvar)
+
+        return ocvars
+
     def generate_unreal_command_line_args(self, map_name):
         command_line_args = f'{self.extra_cmdline_args_setting}'
 
@@ -2196,8 +2235,7 @@ class DeviceUnreal(Device):
             dp_cvars.append(self.dpcvar_for_mediaprofile(mediaprofile_gamepath))
 
         # Add user set dp cvars, overriding any of the forced ones.
-        user_dp_cvars = DeviceUnreal.csettings["dp_cvars"].get_value(self.name)
-        user_dp_cvars = [cvar.strip() for cvar in user_dp_cvars if len(cvar.strip().split('=')) == 2]
+        user_dp_cvars = self.conform_dpcvars(DeviceUnreal.csettings["dp_cvars"].get_value(self.name))
         dp_cvars = self.add_or_override_cvars(dp_cvars, user_dp_cvars)
 
         # add accumulated dpcvars to args
