Switchboard: Added a bulk edit option to string list settings (e.g. dpcvars).

The stage crew was observed to keep a text file with a template of settings that they copy-paste one by one. This will offer the option to copy-paste directly all the items at once in such cases.

#rnx

diff --git a/Engine/Plugins/VirtualProduction/Switchboard/Source/Switchboard/switchboard/config.py b/Engine/Plugins/VirtualProduction/Switchboard/Source/Switchboard/switchboard/config.py
--- a/Engine/Plugins/VirtualProduction/Switchboard/Source/Switchboard/switchboard/config.py
+++ b/Engine/Plugins/VirtualProduction/Switchboard/Source/Switchboard/switchboard/config.py
@@ -1446,6 +1446,7 @@ class StringListSetting(ListSetting):
     An array setting where the elements are strings
     '''
 
+    @override
     def create_element(self, override_device_name: str, index: int) -> Tuple[QtWidgets.QWidget, object]:
         line_edit = QtWidgets.QLineEdit()
         line_edit.editingFinished.connect(
@@ -1454,16 +1455,104 @@ class StringListSetting(ListSetting):
         )
         return line_edit, ""
 
-    def _on_editor_widget_changed(self, line_edit, override_device_name, index):
+    def _on_editor_widget_changed(self, line_edit: QtWidgets.QLineEdit, override_device_name: str, index: int):
+        '''
+        Updates the value in the list when the editor widget's value changes.
+
+        Args:
+            line_edit (QtWidgets.QLineEdit): The line edit widget containing the new value.
+            override_device_name (str): The device name for which the value is being updated.
+            index (int): The index of the element being updated.
+        '''
         current_list = self.get_value(override_device_name)
+
         # List needs to a different object otherwise _on_widget_value_changed will think the value has not changed
         copied_value = current_list.copy()
         copied_value[index] = line_edit.text()
         self._on_widget_value_changed(copied_value, override_device_name)
 
+    @override
     def update_element_value(self, editor_widget: QtWidgets.QLineEdit, list_value: list, index: int):
+        '''
+        Updates the editor widget with the value from the list.
+
+        Args:
+            editor_widget (QtWidgets.QLineEdit): The widget to be updated.
+            list_value (list): The list of values.
+            index (int): The index of the value to be set in the widget.
+        '''
         editor_widget.setText(list_value[index])
 
+    @override
+    def _create_header(self, override_device_name: str) -> QtWidgets.QWidget:
+        '''
+        Creates the header widget for the list setting, including the bulk edit button.
+
+        Args:
+            override_device_name (str): The device name for which the header is being created.
+
+        Returns:
+            QtWidgets.QWidget: The header widget.
+        '''
+        header = super()._create_header(override_device_name)
+        header_layout = header.layout()
+
+        # Use the edit icon to represent Bulk Edit. Contrast with the plus button to add new elements.
+        pixmap = QtGui.QPixmap(":icons/images/icon_edit.png")
+
+        bulk_edit_button = QtWidgets.QPushButton()
+        bulk_edit_button.setIcon(QtGui.QIcon(pixmap))
+        bulk_edit_button.setToolTip("Bulk Edit")
+        bulk_edit_button.setFlat(True)
+        bulk_edit_button.setSizePolicy(QtWidgets.QSizePolicy.Maximum, QtWidgets.QSizePolicy.Maximum)
+        bulk_edit_button.setMaximumWidth(12)
+        bulk_edit_button.setMaximumHeight(12)
+        bulk_edit_button.setFocusPolicy(QtCore.Qt.FocusPolicy.NoFocus)
+        bulk_edit_button.pressed.connect(
+            lambda override_device_name=override_device_name:
+                self._on_press_bulk_edit(override_device_name)
+        )
+
+        header_layout.insertWidget(2, bulk_edit_button)  # Insert 'E' button after '+' button
+
+        # Add a separator item to avoid pressing the clear button by mistake when wanting to bulk edit.
+        separator = QtWidgets.QSpacerItem(12, 0, QtWidgets.QSizePolicy.Fixed, QtWidgets.QSizePolicy.Minimum)
+        header_layout.insertItem(3, separator)
+
+        return header
+
+    def _on_press_bulk_edit(self, override_device_name: str):
+        '''
+        Opens a dialog to allow bulk editing of the list of strings.
+
+        Args:
+            override_device_name (str): The device name for which the bulk edit is being performed.
+        '''
+
+        current_list = self.get_value(override_device_name)
+
+        # Remove empty elements from the list
+        filtered_list = [item for item in current_list if item.strip()]
+
+        dialog = QtWidgets.QDialog()
+        dialog.setWindowTitle(f"Bulk Edit {self.nice_name}")
+        layout = QtWidgets.QVBoxLayout(dialog)
+
+        text_edit = QtWidgets.QPlainTextEdit()
+        text_edit.setPlainText("\n".join(filtered_list))
+        layout.addWidget(text_edit)
+
+        button_box = QtWidgets.QDialogButtonBox(QtWidgets.QDialogButtonBox.Ok | QtWidgets.QDialogButtonBox.Cancel)
+        button_box.accepted.connect(dialog.accept)
+        button_box.rejected.connect(dialog.reject)
+        layout.addWidget(button_box)
+
+        result = dialog.exec()
+        if result == QtWidgets.QDialog.Accepted:
+            new_text = text_edit.toPlainText()
+            new_list = [line for line in new_text.splitlines() if line.strip()]
+            self._on_widget_value_changed(new_list, override_device_name)
+
 
 class LoggingModel(QtGui.QStandardItemModel):
     '''
