Switchboard: Fixed SettingsDialog not being GC'd.

The lambda of its finished signal captured self which causes a circular reference and never let it be garbage collected. Fixed by not using a lambda and instead caching the config that it was trying to capture.

Also fixed a similar situation with Setting widgets. In that case it was fixed by capturing a weak reference to self.

#rnx

diff --git a/Engine/Plugins/VirtualProduction/Switchboard/Source/Switchboard/switchboard/config.py b/Engine/Plugins/VirtualProduction/Switchboard/Source/Switchboard/switchboard/config.py
--- a/Engine/Plugins/VirtualProduction/Switchboard/Source/Switchboard/switchboard/config.py
+++ b/Engine/Plugins/VirtualProduction/Switchboard/Source/Switchboard/switchboard/config.py
@@ -1,19 +1,21 @@
 # Copyright Epic Games, Inc. All Rights Reserved.
 
 import collections
-from enum import Enum
 import fnmatch
 import io
 import json
 import os
 import pathlib
 import shutil
-from pathlib import Path, PurePosixPath
 import socket
 import sys
 import threading
 import time
-from typing import Any, Callable, Optional, Tuple, Type, Union, List
+import weakref
+
+from enum import Enum
+from pathlib import Path, PurePosixPath
+from typing import (Any, Callable, Optional, Tuple, Type, Union, List)
 
 from PySide6 import QtCore
 from PySide6 import QtGui
@@ -583,32 +585,42 @@ class Setting(QtCore.QObject):
     def get_create_ui_batch_manager():
         return SettingsBatchManagerToCreateUI.get_instance()
 
+    def _pre_on_setting_changed(self, old_value, new_value, override_device_name):
+        ''' Called when the base value of the setting changes. '''
+
+        # Since this is a notification of base value changed, don't call _on_setting_changed if the value is being
+        # overridden for the given device name. Otherwise _on_setting_changed may incorrectly update the override
+        # widget with the base value.
+        if not self.is_overridden(override_device_name):
+            self._on_setting_changed(new_value, override_device_name=override_device_name)
+
     def _register_on_setting_changed(self, top_level_widget: QtWidgets.QWidget, override_device_name: str):
+        ''' Registers the callbacks when the value of the setting changes. '''
 
-        def on_setting_changed(old_value, new_value, override_device_name=override_device_name):
+        # Register the "pre" change which will filter the call to _on_setting_changed
 
-            # Since this is a notification of base value changed, don't call _on_setting_changed if the value is being
-            # overridden for the given device name. Otherwise _on_setting_changed may incorrectly update the override
-            # widget with the base value.
-            if not self.is_overridden(override_device_name):
-                self._on_setting_changed(new_value, override_device_name=override_device_name)
+        self_weak_ref = weakref.ref(self)  # Use weak to avoid holding a strong reference to the object
 
-        self.signal_setting_changed.connect(on_setting_changed)
+        def pre_on_setting_changed_weak(old_value, new_value):
+            ''' Captures override_device_name but not self to avoid memory leaks due to strong reference '''
+            self_instance = self_weak_ref()
+            if self_instance is not None:
+                self_instance._pre_on_setting_changed(old_value, new_value, override_device_name)
 
-        # Clear the widget when it is destroyed to avoid dangling references
+        self.signal_setting_changed.connect(pre_on_setting_changed_weak)
 
-        def handle_widget_destroyed(destroyed_object=None):
-            self._on_widget_destroyed(on_setting_changed, override_device_name)
+        # Clear the widget when it is destroyed
 
-        top_level_widget.destroyed.connect(handle_widget_destroyed)
+        def on_widget_destroyed_weak(destroyed_object=None):
+            ''' Captures override_device_name but not self to avoid memory leaks due to strong reference '''
+            self_instance = self_weak_ref()
+            if self_instance is not None:
+                self_instance._on_widget_destroyed(override_device_name)
 
-    def _on_widget_destroyed(self, on_setting_changed_lambda, override_device_name: str):
-        try:
-            self.signal_setting_changed.disconnect(on_setting_changed_lambda)
-        except Exception as exc:
-            LOGGER.warning('Failed to disconnect on_setting_changed_lambda',
-                           exc_info=exc)
+        top_level_widget.destroyed.connect(on_widget_destroyed_weak)
 
+    def _on_widget_destroyed(self, destroyed_object=None, override_device_name: str = None):
+        ''' Called when the widget associated with the given override_device_name is destroyed. '''
         self.set_widget(widget=None, override_device_name=override_device_name)
 
     def _decorate_with_reset_widget(self, override_device_name: str, setting_editor_widget: QtWidgets.QWidget):
diff --git a/Engine/Plugins/VirtualProduction/Switchboard/Source/Switchboard/switchboard/settings_dialog.py b/Engine/Plugins/VirtualProduction/Switchboard/Source/Switchboard/switchboard/settings_dialog.py
--- a/Engine/Plugins/VirtualProduction/Switchboard/Source/Switchboard/switchboard/settings_dialog.py
+++ b/Engine/Plugins/VirtualProduction/Switchboard/Source/Switchboard/switchboard/settings_dialog.py
@@ -90,6 +90,7 @@ class SettingsDialog(QtCore.QObject):
     def __init__(self, settings: UserSettings, config: Config):
         super().__init__()
         self.plugin_widgets = {}
+        self.config = config
 
         # Set the UI object
         self.ui = QtWidgets.QDialog()
@@ -97,10 +98,7 @@ class SettingsDialog(QtCore.QObject):
         dialog_layout = QtWidgets.QVBoxLayout(self.ui)
         dialog_layout.setContentsMargins(2, 2, 2, 2)
         self.ui.setWindowTitle("Settings")
-        self.ui.finished.connect(
-            lambda result, config=config:
-                self._on_finished(result, config)
-        )
+        self.ui.finished.connect(self._on_finished)
 
         self._create_search_area(dialog_layout)
         self.general_settings_list = [
@@ -172,8 +170,8 @@ class SettingsDialog(QtCore.QObject):
         self._search_settings_again()
         return
 
-    def _on_finished(self, result: int, config: Config):
-        config.P4_ENABLED.signal_setting_changed.disconnect(self._on_source_control_setting_changed)
+    def _on_finished(self, result: int):
+        self.config.P4_ENABLED.signal_setting_changed.disconnect(self._on_source_control_setting_changed)
 
         # Currently, the only way to dismiss the settings dialog is by using
         # the close button as opposed to ok/cancel buttons, so we intercept the
@@ -182,7 +180,7 @@ class SettingsDialog(QtCore.QObject):
         if (self._changed_config_path in self._config_paths
                 and self._changed_config_path != self._current_config_path):
             # Show the confirmation dialog using a relative path to the config.
-            rel_config_path = config.get_relative_config_path(
+            rel_config_path = self.config.get_relative_config_path(
                 self._changed_config_path)
             reply = QtWidgets.QMessageBox.question(
                 self.ui, 'Confirm Overwrite',
diff --git a/Engine/Plugins/VirtualProduction/Switchboard/Source/Switchboard/switchboard/switchboard_widgets.py b/Engine/Plugins/VirtualProduction/Switchboard/Source/Switchboard/switchboard/switchboard_widgets.py
--- a/Engine/Plugins/VirtualProduction/Switchboard/Source/Switchboard/switchboard/switchboard_widgets.py
+++ b/Engine/Plugins/VirtualProduction/Switchboard/Source/Switchboard/switchboard/switchboard_widgets.py
@@ -465,9 +465,6 @@ class AddressComboBox(NonScrollableComboBox):
         if self.resolving:
             return False
 
-        if self.runnable:
-            self.runnable.signals.result.disconnect()
-
         name = self.lineEdit().text()
         self.runnable = AddressComboBox.ResolverRunnable(name)
         self.runnable.signals.result.connect(self.on_resolver_result)
