Switchboard: Fixed memory leaks due to circular references. In particular, DeviceWidgetUnreal.

The issue stemmed from:
* Objects owning settings that have signals connected to lambdas that capture the owning object. This was fixed with weak references.
* Devices own the listener client that has signals connected to the owning object. This was fixed by making the listener client a qobject and settings its parent to the owning device.

Also added debug_utils.py that can register a Debug menu in the main toolbar to make it easier to trigger non-intrusive test code. Includes a convenience function to log instances of a given class. SB doesn't currently include objgraph but that is a useful package to debug circular dependencies. pip install objgraph, install GraphViz, use objgraph.show_backrefs(..) and it will generate a png with the dependecy diagram of the object.

#rnx

diff --git a/Engine/Plugins/VirtualProduction/Switchboard/Source/Switchboard/switchboard/debug_utils.py b/Engine/Plugins/VirtualProduction/Switchboard/Source/Switchboard/switchboard/debug_utils.py
new file mode 100644
--- /dev/null
+++ b/Engine/Plugins/VirtualProduction/Switchboard/Source/Switchboard/switchboard/debug_utils.py
@@ -0,0 +1,104 @@
+# Copyright Epic Games, Inc. All Rights Reserved.
+
+from __future__ import annotations
+
+import gc
+
+from PySide6.QtWidgets import QWidgetAction, QMenu
+from typing import Callable, List
+
+from switchboard.switchboard_logging import LOGGER
+
+# Set as True to add a Debug menut to the Main Dialog
+ADD_MENU = False
+
+
+def test():
+    ''' Placeholder for debugging tests. Will be called from Tools Menu item. '''
+    pass
+
+
+def register_menu_actions(menu: QMenu) -> None:
+    '''
+    Registers menu actions to the given QMenu object.
+
+    This function adds a new action labeled "Test" to the provided QMenu. 
+    When the action is triggered, the `test()` function is called.
+
+    Args:
+        menu (QMenu): The QMenu object to which the action should be added.
+    '''
+
+    action = register_menu_action(menu, "&Test")
+    action.triggered.connect(lambda: test())
+
+
+def register_menu_action(
+        current_menu: QMenu,
+        actionname: str,
+        menunames: List[str] = [],
+        actiontooltip: str = '') -> QWidgetAction:
+    ''' Registers a QWidgetAction with the tools menu
+
+    Args:
+        current_menu: Current menu we're adding to.
+        actionname: Name of the action to be added
+        menunames: Submenu(s) where to place the given action
+
+    Returns:
+        QWidgetAction: The action that was created.
+    '''
+
+    # Find find submenu, and create the ones that don't exist along the way
+
+    # iterate over menunames given
+    for menuname in menunames:
+
+        create_menu = True
+
+        # try to find existing menu
+        for current_action in current_menu.actions():
+
+            submenu = current_action.menu()
+
+            if submenu and (menuname == submenu.title()):
+                current_menu = submenu
+                create_menu = False
+                break
+
+        # if we didn't find the submenu, create it
+        if create_menu:
+            newmenu = QMenu(parent=current_menu)
+            newmenu.setTitle(menuname)
+            current_menu.addMenu(newmenu)
+            current_menu = newmenu
+
+    # add the given action
+    action = QWidgetAction(current_menu)
+    action.setText(actionname)
+    action.setToolTip(actiontooltip)
+    current_menu.setToolTipsVisible(True)
+    current_menu.addAction(action)
+
+    return action
+
+
+def print_instances(cls: type, print_func: Callable[[object], str]) -> None:
+    ''' Prints the number of instances and the identifying information of all instances of the given class,
+    using the provided print_func for displaying identifying information. Useful when tracking down memory leaks.
+
+    Example:
+        print_instances(BoolSetting, lambda setting: setting.nice_name)
+
+    Args:
+        cls: The class type for which to print instances.
+        print_func: A function that takes an object of type `cls` and returns a string for identifying it.
+    '''
+    instances = [obj for obj in gc.get_objects() if isinstance(obj, cls)]
+
+    # Print the number of instances
+    LOGGER.debug(f"{cls.__name__}: {len(instances)} instances")
+
+    # Print the instances using the provided print_func
+    for instance in sorted(instances, key=print_func):
+        LOGGER.debug(f"{cls.__name__} instance: {print_func(instance)}")
diff --git a/Engine/Plugins/VirtualProduction/Switchboard/Source/Switchboard/switchboard/device_list_widget.py b/Engine/Plugins/VirtualProduction/Switchboard/Source/Switchboard/switchboard/device_list_widget.py
--- a/Engine/Plugins/VirtualProduction/Switchboard/Source/Switchboard/switchboard/device_list_widget.py
+++ b/Engine/Plugins/VirtualProduction/Switchboard/Source/Switchboard/switchboard/device_list_widget.py
@@ -1,14 +1,14 @@
 # Copyright Epic Games, Inc. All Rights Reserved.
 
-from typing import Dict, List
-
-from PySide6 import QtCore, QtGui, QtWidgets
+import weakref
 
-from switchboard.devices.device_base import Device, DeviceStatus, \
-    PluginHeaderWidgets
-from switchboard.devices.device_widget_base import DeviceWidget, DeviceAutoJoinMUServerUI
 import switchboard.switchboard_widgets as sb_widgets
 
+from typing import Dict, List
+from PySide6 import QtCore, QtGui, QtWidgets
+from switchboard.devices.device_base import Device, DeviceStatus, PluginHeaderWidgets
+from switchboard.devices.device_widget_base import DeviceAutoJoinMUServerUI
+
 
 class DeviceListWidget(QtWidgets.QListWidget):
     signal_register_device_widget = QtCore.Signal(object)
@@ -148,11 +148,16 @@ class DeviceListWidget(QtWidgets.QListWidget):
         header_widget.autojoin_mu.get_button().blockSignals(wasBlockingSignals)
 
     def connect_to_device_autojoin(self, device, category):
-        def child_device_changed():
-            self.update_header_autojoin_state(category)
+        ''' Connects to the device autojoin setting with a weak reference to avoid strong reference cycles. '''
+        self_weak_ref = weakref.ref(self)
+
+        def child_device_changed_weak():
+            self_instance = self_weak_ref()
+            if self_instance is not None:
+                self_instance.update_header_autojoin_state(category)
 
         if device.autojoin_mu_server:
-            device.autojoin_mu_server.signal_setting_changed.connect(child_device_changed)
+            device.autojoin_mu_server.signal_setting_changed.connect(child_device_changed_weak)
             self.update_header_autojoin_state(category)
 
     def attach_autojoin_button(self, header_widget, category):
diff --git a/Engine/Plugins/VirtualProduction/Switchboard/Source/Switchboard/switchboard/devices/device_base.py b/Engine/Plugins/VirtualProduction/Switchboard/Source/Switchboard/switchboard/devices/device_base.py
--- a/Engine/Plugins/VirtualProduction/Switchboard/Source/Switchboard/switchboard/devices/device_base.py
+++ b/Engine/Plugins/VirtualProduction/Switchboard/Source/Switchboard/switchboard/devices/device_base.py
@@ -2,11 +2,13 @@
 
 from __future__ import annotations
 
-from enum import IntEnum, IntFlag, auto, unique
 import os
 import random
 import socket
 import threading
+import weakref
+
+from enum import IntEnum, IntFlag, auto, unique
 from typing import Optional, Type, Dict, OrderedDict, TYPE_CHECKING
 
 from PySide6 import QtCore
@@ -90,8 +92,10 @@ class Device(QtCore.QObject):
         self.device_qt_handler = DeviceQtHandler()
         self.autojoin_mu_server = False
         self.setting_address = AddressSetting('address', 'Address', address, category="Network Settings")
+
+        weak_self = weakref.ref(self)
         self.setting_address.signal_setting_changed.connect(
-            lambda: self._check_recreate_osc_client())
+            lambda: weak_self() and weak_self()._check_recreate_osc_client())
 
         # override any setting that was passed via kwargs
         for setting in self.setting_overrides():
@@ -120,15 +124,22 @@ class Device(QtCore.QObject):
         self.widget: DeviceWidget = None  # Constructed during init()
 
     def init(self, widget_class, icons):
-        self.widget = widget_class(
-            self.name, self.device_hash, self.address, icons)
-        self.widget.signal_device_name_changed.connect(
-            self.on_device_name_changed)
-        self.widget.signal_address_changed.connect(
-            self.on_address_changed)
-        self.setting_address.signal_setting_changed.connect(
-            lambda _, new_address, widget=self.widget:
-                widget.on_address_changed(new_address))
+
+        self.widget = widget_class(self.name, self.device_hash, self.address, icons)
+
+        self.widget.signal_device_name_changed.connect(self.on_device_name_changed)
+        self.widget.signal_address_changed.connect(self.on_address_changed)
+
+        # weakref to void an unintended strong reference that can prevent garbage collection
+
+        self_weak_ref = weakref.ref(self)
+
+        def on_address_changed_weak(old_address, new_address):
+            self_instance = self_weak_ref()
+            if self_instance is not None and self_instance.widget is not None:
+                self_instance.widget.on_address_changed(new_address)
+
+        self.setting_address.signal_setting_changed.connect(on_address_changed_weak)
 
         self.disable_override_style_on_device_only_settings()
 
diff --git a/Engine/Plugins/VirtualProduction/Switchboard/Source/Switchboard/switchboard/devices/unreal/plugin_unreal.py b/Engine/Plugins/VirtualProduction/Switchboard/Source/Switchboard/switchboard/devices/unreal/plugin_unreal.py
--- a/Engine/Plugins/VirtualProduction/Switchboard/Source/Switchboard/switchboard/devices/unreal/plugin_unreal.py
+++ b/Engine/Plugins/VirtualProduction/Switchboard/Source/Switchboard/switchboard/devices/unreal/plugin_unreal.py
@@ -3,33 +3,36 @@
 from __future__ import annotations
 
 import base64
-from collections import OrderedDict
 import concurrent.futures
-from datetime import datetime
-from functools import wraps
 import hashlib
-from ipaddress import IPv4Address
-from itertools import count
 import json
 import os
 import pathlib
-from pathlib import PurePosixPath
 import re
 import socket
 import sys
 import threading
-from typing import Callable, Dict, Generator, Optional, Union
-import uuid
 import time
+import uuid
+import weakref
+
+from collections import OrderedDict
+from datetime import datetime
 from enum import Enum
+from functools import wraps
+from ipaddress import IPv4Address
+from itertools import count
+from pathlib import PurePosixPath
+from typing import Callable, Dict, Generator, Optional, Union
 
 from PySide6 import QtCore, QtGui, QtWidgets
 
-from switchboard import message_protocol, switchboard_application, ugs_utils
 import switchboard.config_osc as osc
 import switchboard.switchboard_dialog as sb_dialog
 import switchboard.switchboard_utils as sb_utils
 import switchboard.switchboard_widgets as sb_widgets
+
+from switchboard import message_protocol, switchboard_application, ugs_utils
 from switchboard.config import (
     CONFIG, DEFAULT_MAP_TEXT, ENABLE_UGS_SUPPORT, SETTINGS, BoolSetting, DirectoryPathSetting,
     EngineSyncMethod, FilePathSetting, IntSetting, MultiOptionSetting, OptionSetting, Setting,
@@ -520,7 +523,7 @@ class DeviceUnreal(Device):
         '/Script/MediaFrameworkUtilities.MediaProfile',
     )
 
-    csettings = {
+    csettings: dict[str, Setting] = {
         'buffer_size': IntSetting(
             attr_name="buffer_size",
             nice_name="Buffer Size",
@@ -966,8 +969,8 @@ class DeviceUnreal(Device):
         self.unreal_client = ListenerClient(
             address=self.address,
             port=DeviceUnreal.csettings['port'].get_value(self.name),
-            buffer_size=DeviceUnreal.csettings['buffer_size'].get_value(
-                self.name)
+            buffer_size=DeviceUnreal.csettings['buffer_size'].get_value(self.name),
+            parent=self,
         )
 
         roles = kwargs.get("roles", [])
@@ -1047,11 +1050,16 @@ class DeviceUnreal(Device):
         CONFIG.ENGINE_SYNC_METHOD.signal_setting_changed.connect(
             self.on_engine_sync_method_changed)
 
+        weak_self = weakref.ref(self)
+
         osc_port_setting = DeviceUnreal.csettings['osc_port']
+
         osc_port_setting.signal_setting_changed.connect(
-            lambda: self._check_recreate_osc_client())
+            lambda: weak_self() and weak_self()._check_recreate_osc_client()
+        )
         osc_port_setting.signal_setting_overridden.connect(
-            lambda: self._check_recreate_osc_client())
+            lambda: weak_self() and weak_self()._check_recreate_osc_client()
+        )
 
         self.auto_connect = False
 
@@ -1071,22 +1079,17 @@ class DeviceUnreal(Device):
             self._on_listener_connection_failed, QtCore.Qt.QueuedConnection)
 
         # Set a delegate method if the device gets a disconnect signal
-        self.unreal_client.disconnect_delegate = self.on_listener_disconnect
-        self.unreal_client.receive_file_completed_delegate = (
-            self.on_file_received)
-        self.unreal_client.receive_file_failed_delegate = (
-            self.on_file_receive_failed)
-        self.unreal_client.delegates["state"] = self.on_listener_state
-        self.unreal_client.delegates["kill"] = self.on_program_killed
-        self.unreal_client.delegates["programstdout"] = (
-            self.on_listener_programstdout)
-        self.unreal_client.delegates["program ended"] = (
-            self.on_program_ended)
-        self.unreal_client.delegates["program started"] = (
-            self.on_program_started)
-        self.unreal_client.delegates["program killed"] = self.on_program_killed
+        self.unreal_client.disconnect_delegate = (self.on_listener_disconnect)
+        self.unreal_client.receive_file_completed_delegate = (self.on_file_received)
+        self.unreal_client.receive_file_failed_delegate = (self.on_file_receive_failed)
+        self.unreal_client.delegates["state"] = (self.on_listener_state)
+        self.unreal_client.delegates["kill"] = (self.on_program_killed)
+        self.unreal_client.delegates["programstdout"] = (self.on_listener_programstdout)
+        self.unreal_client.delegates["program ended"] = (self.on_program_ended)
+        self.unreal_client.delegates["program started"] = (self.on_program_started)
+        self.unreal_client.delegates["program killed"] = (self.on_program_killed)
         # This catches start failures.
-        self.unreal_client.delegates["start"] = self.on_program_started
+        self.unreal_client.delegates["start"] = (self.on_program_started)
 
         self.osc_connection_timer = QtCore.QTimer(self)
         self.osc_connection_timer.timeout.connect(self._try_connect_osc)
@@ -1153,13 +1156,15 @@ class DeviceUnreal(Device):
     def init(self, widget_class, icons):
         super().init(widget_class, icons)
 
+        weak_self = weakref.ref(self)
+
         self.exclude_from_build.signal_setting_changed.connect(
-            lambda _, new_value: self.on_setting_exclude_from_build_changed(new_value)
+            lambda _, new_value: weak_self() and weak_self().on_setting_exclude_from_build_changed(new_value)
         )
         self.on_setting_exclude_from_build_changed(self.exclude_from_build.get_value())
 
         self.exclude_from_insights.signal_setting_changed.connect(
-            lambda _, new_value: self.on_setting_exclude_from_insights_changed(new_value)
+            lambda _, new_value: weak_self() and weak_self().on_setting_exclude_from_insights_changed(new_value)
         )
         self.on_setting_exclude_from_insights_changed(self.exclude_from_insights.get_value())
 
diff --git a/Engine/Plugins/VirtualProduction/Switchboard/Source/Switchboard/switchboard/devices/unreal/redeploy_dialog.py b/Engine/Plugins/VirtualProduction/Switchboard/Source/Switchboard/switchboard/devices/unreal/redeploy_dialog.py
--- a/Engine/Plugins/VirtualProduction/Switchboard/Source/Switchboard/switchboard/devices/unreal/redeploy_dialog.py
+++ b/Engine/Plugins/VirtualProduction/Switchboard/Source/Switchboard/switchboard/devices/unreal/redeploy_dialog.py
@@ -41,7 +41,7 @@ class RedeployListenerEndpoint(QtCore.QObject):
         self.address = address
         self.port = port
 
-        self.client = ListenerClient(self.address, self.port)
+        self.client = ListenerClient(self.address, self.port, parent=self)
         self.client.listener_qt_handler.listener_connecting.connect(
             lambda _: self.version_label.setText('Connecting...'))
         self.client.listener_qt_handler.listener_connection_failed.connect(
diff --git a/Engine/Plugins/VirtualProduction/Switchboard/Source/Switchboard/switchboard/listener_client.py b/Engine/Plugins/VirtualProduction/Switchboard/Source/Switchboard/switchboard/listener_client.py
--- a/Engine/Plugins/VirtualProduction/Switchboard/Source/Switchboard/switchboard/listener_client.py
+++ b/Engine/Plugins/VirtualProduction/Switchboard/Source/Switchboard/switchboard/listener_client.py
@@ -2,24 +2,26 @@
 
 from __future__ import annotations
 
-from abc import abstractmethod, ABC
 import asyncio
-from asyncio import Future
-from collections import deque
-from datetime import datetime, timedelta
 import logging
 import ssl
 import threading
 import traceback
-from typing import cast, Callable, Dict, Optional
 import uuid
 
+from abc import ABC, abstractmethod
+from asyncio import Future
+from collections import deque
+from datetime import datetime, timedelta
+from typing import Callable, Dict, Optional, cast
+
 from aioquic.asyncio.client import connect
 from aioquic.asyncio.protocol import QuicConnectionProtocol
 from aioquic.quic.configuration import QuicConfiguration
 from aioquic.quic.events import QuicEvent, HandshakeCompleted
 from cryptography.hazmat.primitives import hashes
 from cryptography.x509.oid import NameOID
+
 from PySide6 import QtCore, QtGui, QtWidgets
 from PySide6.QtWidgets import QMessageBox
 
@@ -233,8 +235,8 @@ class ListenerQtHandler(QtCore.QObject):
             assert self.future
             self.future.set_result(result)
 
-    def __init__(self):
-        super().__init__()
+    def __init__(self, parent=None):
+        super().__init__(parent)
 
         self._queued_prompts: list[ListenerQtHandler.Prompt] = []
 
@@ -278,7 +280,7 @@ class ListenerQtHandler(QtCore.QObject):
         self._queued_prompts.clear()
 
 
-class ListenerClient:
+class ListenerClient(QtCore.QObject):
     '''Connects to a server running SwitchboardListener.
 
     Runs a thread to service its socket, and upon receiving complete messages,
@@ -301,7 +303,9 @@ class ListenerClient:
                            + EXTRA_HITCH_TOLERANCE_SEC)
     IDLE_TIMEOUT_SEC = 5.0
 
-    def __init__(self, address, port, buffer_size=1024):
+    def __init__(self, address, port, buffer_size=1024, parent=None):
+        super().__init__(parent)
+
         self.address = address
         self.port = port
         self.buffer_size = buffer_size
diff --git a/Engine/Plugins/VirtualProduction/Switchboard/Source/Switchboard/switchboard/switchboard_dialog.py b/Engine/Plugins/VirtualProduction/Switchboard/Source/Switchboard/switchboard/switchboard_dialog.py
--- a/Engine/Plugins/VirtualProduction/Switchboard/Source/Switchboard/switchboard/switchboard_dialog.py
+++ b/Engine/Plugins/VirtualProduction/Switchboard/Source/Switchboard/switchboard/switchboard_dialog.py
@@ -46,6 +46,8 @@ from switchboard.devices.unreal.redeploy_dialog import RedeployListenerDialog
 from switchboard.util import collect_logs
 from switchboard.ui.switchboard_ui import Ui_switchboard_window
 
+import switchboard.debug_utils as debug_utils
+
 ENGINE_PATH = "../../../../.."
 RELATIVE_PATH = os.path.dirname(__file__)
 EMPTY_SYNC_ENTRY = "-- None --"
@@ -263,6 +265,7 @@ class SwitchboardDialog(QtCore.QObject):
     def __init__(self, script_manager):
         super().__init__()
 
+        self.window = None
         self.script_manager = script_manager
 
         font_dir = os.path.join(ENGINE_PATH, 'Content/Slate/Fonts')
@@ -321,6 +324,11 @@ class SwitchboardDialog(QtCore.QObject):
         self.register_open_logs_menuitem()
         self.register_zip_logs_menuitem()
 
+        # Debugging menu item
+        if debug_utils.ADD_MENU:
+            self.debug_menu = self.window.menu_bar.addMenu("&Debug")
+            debug_utils.register_menu_actions(self.debug_menu)
+
         # Redeploy menu item
         self.register_redeploy_menuitem()
 
