Switchboard: Fix not handling nDisplay config file parsing.

When this happens, it puts SB in a condition in which loading new configs doesn't work correctly.

#rnx

diff --git a/Engine/Plugins/VirtualProduction/Switchboard/Source/Switchboard/switchboard/devices/ndisplay/plugin_ndisplay.py b/Engine/Plugins/VirtualProduction/Switchboard/Source/Switchboard/switchboard/devices/ndisplay/plugin_ndisplay.py
--- a/Engine/Plugins/VirtualProduction/Switchboard/Source/Switchboard/switchboard/devices/ndisplay/plugin_ndisplay.py
+++ b/Engine/Plugins/VirtualProduction/Switchboard/Source/Switchboard/switchboard/devices/ndisplay/plugin_ndisplay.py
@@ -292,7 +292,7 @@ class AddnDisplayDialog(AddDeviceDialog):
 
             return devices
 
-        except (IndexError, KeyError):
+        except (IndexError, KeyError, FileNotFoundError):
             LOGGER.error(f"Error parsing nDisplay config file {cfg_file}")
             return []
 
@@ -1683,7 +1683,11 @@ class DevicenDisplay(DeviceUnreal):
         '''
         Updates settings that are exclusively controlled by the config file.
         '''
-        config = self.__class__.parse_config(cfg_file)
+        try:
+            config = self.__class__.parse_config(cfg_file)
+        except Exception as e:
+            LOGGER.error(f"Error parsing config file {cfg_file}: {e}")
+            return
 
         self.bp_object_path = config.uasset_path
         nodes = config.nodes
