PR #11858: Fix: [Switchboard] When project/engine changelist changes, throw exception bug.

#rnx

diff --git a/Engine/Plugins/VirtualProduction/Switchboard/Source/Switchboard/switchboard/switchboard_dialog.py b/Engine/Plugins/VirtualProduction/Switchboard/Source/Switchboard/switchboard/switchboard_dialog.py
--- a/Engine/Plugins/VirtualProduction/Switchboard/Source/Switchboard/switchboard/switchboard_dialog.py
+++ b/Engine/Plugins/VirtualProduction/Switchboard/Source/Switchboard/switchboard/switchboard_dialog.py
@@ -1729,7 +1729,7 @@ class SwitchboardDialog(QtCore.QObject):
         self._project_changelist = value
 
         if self.window.project_cl_combo_box.currentText() != self._project_changelist:
-            self.window.project_cl_combo_box.setText(self._project_changelist)
+            self.window.project_cl_combo_box.setCurrentText(self._project_changelist)
 
         # Check if all of the devices are on the right changelist
         for device in self.device_manager.devices():
@@ -1752,7 +1752,7 @@ class SwitchboardDialog(QtCore.QObject):
         self._engine_changelist = changelist_value
 
         if self.window.engine_cl_combo_box.currentText() != self._engine_changelist:
-            self.window.engine_cl_combo_box.setText(self._engine_changelist)
+            self.window.engine_cl_combo_box.setCurrentText(self._engine_changelist)
 
         # Check if all of the devices are on the right changelist
         for device in self.device_manager.devices():
