Listen to overridden value changes to update the IP address. Also make sure we specify an override when we update the address value on the device.

diff --git a/Engine/Plugins/VirtualProduction/Switchboard/Source/Switchboard/switchboard/devices/device_base.py b/Engine/Plugins/VirtualProduction/Switchboard/Source/Switchboard/switchboard/devices/device_base.py
--- a/Engine/Plugins/VirtualProduction/Switchboard/Source/Switchboard/switchboard/devices/device_base.py
+++ b/Engine/Plugins/VirtualProduction/Switchboard/Source/Switchboard/switchboard/devices/device_base.py
@@ -140,6 +140,9 @@ class Device(QtCore.QObject):
                 self_instance.widget.on_address_changed(new_address)
 
         self.setting_address.signal_setting_changed.connect(on_address_changed_weak)
+        self.setting_address.signal_setting_overridden.connect(
+            lambda _, old_address, override, address_changed=on_address_changed_weak:
+                on_address_changed_weak(old_address,override))
 
         self.disable_override_style_on_device_only_settings()
 
@@ -243,7 +246,7 @@ class Device(QtCore.QObject):
 
     @address.setter
     def address(self, value: str):
-        self.setting_address.update_value(value)
+        self.setting_address.override_value(self.name, value)
 
     @property
     def status(self):
