Switchboard: Overwrite of remote NDC from local is now opt-in.

Option is called : Overwrite nDisplay UAsset

#rnx

diff --git a/Engine/Plugins/VirtualProduction/Switchboard/Source/Switchboard/switchboard/devices/ndisplay/plugin_ndisplay.py b/Engine/Plugins/VirtualProduction/Switchboard/Source/Switchboard/switchboard/devices/ndisplay/plugin_ndisplay.py
--- a/Engine/Plugins/VirtualProduction/Switchboard/Source/Switchboard/switchboard/devices/ndisplay/plugin_ndisplay.py
+++ b/Engine/Plugins/VirtualProduction/Switchboard/Source/Switchboard/switchboard/devices/ndisplay/plugin_ndisplay.py
@@ -410,6 +410,13 @@ class DevicenDisplay(DeviceUnreal):
             file_path_filter="Programs and Scripts (*.exe;*.bat;*.sh);;All Files (*)",
             category="General Settings",
         ),
+        'overwrite_remote_ndc_asset': BoolSetting(
+            attr_name="overwrite_remote_ndc_asset",
+            nice_name="Overwrite nDisplay UAsset",
+            tool_tip="Overwrite the remote nDisplay .uasset from the local one when launching, if it exists",
+            value=False,
+            category="Network Settings",
+        ),
         'use_all_available_cores': BoolSetting(
             attr_name="use_all_available_cores",
             nice_name="Use All Available Cores",
@@ -1800,8 +1807,10 @@ class DevicenDisplay(DeviceUnreal):
             cfg_content, cfg_destination)
         self.unreal_client.send_message(cfg_msg)
 
-        # Send .uasset but only if not launching packaged
-        if self.bp_object_path and not self.is_packaged_game():
+        overwrite_ndc = DevicenDisplay.csettings['overwrite_remote_ndc_asset'].get_value(self.name)
+
+        # When desired, send .uasset but only if not launching packaged
+        if overwrite_ndc and self.bp_object_path and not self.is_packaged_game():
             local_project_path = os.path.dirname(
                 CONFIG.UPROJECT_PATH.get_value())
             dest_project_path = os.path.dirname(
