Switchboard: Fixed widgets not unsubscribing from property.

Positional argument of the ondestroyed signal was incorrectly placed.

#rnx

diff --git a/Engine/Plugins/VirtualProduction/Switchboard/Source/Switchboard/switchboard/config.py b/Engine/Plugins/VirtualProduction/Switchboard/Source/Switchboard/switchboard/config.py
--- a/Engine/Plugins/VirtualProduction/Switchboard/Source/Switchboard/switchboard/config.py
+++ b/Engine/Plugins/VirtualProduction/Switchboard/Source/Switchboard/switchboard/config.py
@@ -615,7 +615,7 @@ class Setting(QtCore.QObject):
             ''' Captures override_device_name but not self to avoid memory leaks due to strong reference '''
             self_instance = self_weak_ref()
             if self_instance is not None:
-                self_instance._on_widget_destroyed(override_device_name)
+                self_instance._on_widget_destroyed(override_device_name=override_device_name)
 
         top_level_widget.destroyed.connect(on_widget_destroyed_weak)
 
