Switchboard: Add -multiprocess to command line of MP nodes.

This should avoid certain race conditions between nodes running on the same machine, such as cleaning of shader working directory.

#rnx

diff --git a/Engine/Plugins/VirtualProduction/Switchboard/Source/Switchboard/switchboard/devices/ndisplay/plugin_ndisplay.py b/Engine/Plugins/VirtualProduction/Switchboard/Source/Switchboard/switchboard/devices/ndisplay/plugin_ndisplay.py
--- a/Engine/Plugins/VirtualProduction/Switchboard/Source/Switchboard/switchboard/devices/ndisplay/plugin_ndisplay.py
+++ b/Engine/Plugins/VirtualProduction/Switchboard/Source/Switchboard/switchboard/devices/ndisplay/plugin_ndisplay.py
@@ -2,14 +2,16 @@
 
 import json
 import os
-import sys
-from pathlib import Path, PurePosixPath
 import socket
 import struct
+import sys
 import traceback
 import uuid
-from typing import Optional, List
+
 from enum import Enum
+from ipaddress import ip_address
+from pathlib import Path, PurePosixPath
+from typing import Optional, List
 
 from PySide6 import QtCore
 from PySide6 import QtWidgets
@@ -49,6 +51,7 @@ class LaunchMode(Enum):
     Standalone = "Standalone (ICVFX)"
     Packaged = "Packaged Game"
 
+
 class AddnDisplayDialog(AddDeviceDialog):
     def __init__(self, existing_devices, parent=None):
         super().__init__(
@@ -374,7 +377,6 @@ class DisplayConfig(object):
         self.uasset_path = ''
 
 
-
 class DevicenDisplay(DeviceUnreal):
 
     add_device_dialog = AddnDisplayDialog
@@ -1178,8 +1180,25 @@ class DevicenDisplay(DeviceUnreal):
             f'{udpm_transport_static}',   # -UDPMESSAGING_TRANSPORT_STATIC=
         ]
 
+        # Add -multiprocess to all except the first co-hosted node, sorted by name.
+        # Used to avoid race conditions when writing or deleting certain files on disk, such as
+        # shader caches.
+
+        shared_ip_connected_nodes = [
+            node for node in self.get_connected_devices()
+            if ip_address(node.address) == ip_address(self.address)
+        ]
+
+        # Launch command is also generated when devices are not yet connected, so empty shared_ip_connected_nodes
+        # is a valid situation when the cluster is not being actually launched.
+        if len(shared_ip_connected_nodes):
+            first_cohosted_node = sorted(shared_ip_connected_nodes, key=lambda node: node.name)[0]
+
+            if self != first_cohosted_node:
+                args.append('-multiprocess')
+
         # Specify multiplayer with listen server mode by arg for clients if enabled
-        if DevicenDisplay.csettings["multiplayer_mode"].get_value() == 'Listen server'and not self.is_primary():
+        if DevicenDisplay.csettings["multiplayer_mode"].get_value() == 'Listen server' and not self.is_primary():
             args.extend([
                 '-dc_replicationserver_type listen'
             ])
