Switchboard: Force-enables (with opt-out) alpha holdout on Editor launches in SB configs that contain nDisplay devices.

This avoids affecting nDisplay launches or mocap usages of SB when they would still use project settings.

#rnx

diff --git a/Engine/Plugins/VirtualProduction/Switchboard/Source/Switchboard/switchboard/devices/unreal/plugin_unreal.py b/Engine/Plugins/VirtualProduction/Switchboard/Source/Switchboard/switchboard/devices/unreal/plugin_unreal.py
--- a/Engine/Plugins/VirtualProduction/Switchboard/Source/Switchboard/switchboard/devices/unreal/plugin_unreal.py
+++ b/Engine/Plugins/VirtualProduction/Switchboard/Source/Switchboard/switchboard/devices/unreal/plugin_unreal.py
@@ -777,6 +777,15 @@ class DeviceUnreal(Device):
                 'wildcards to be applied during sync.'),
             category="Source Control Settings",
         ),
+        'force_holdout_if_ndisplay': BoolSetting(
+            attr_name='force_holdout_if_ndisplay',
+            nice_name='Force nDisplay Editor Holdout',
+            value=True,
+            tool_tip=(
+                'When checked, force-enables Primitive Alpha Holdout on Editor launches when there are nDisplay devices in '
+                'this Switchboard configuration \n'),
+            category="UE Settings",
+        ),
     }
 
     if ENABLE_UGS_SUPPORT:
@@ -2239,6 +2248,14 @@ class DeviceUnreal(Device):
         if mediaprofile_gamepath:
             dp_cvars.append(self.dpcvar_for_mediaprofile(mediaprofile_gamepath))
 
+        # Holdout composite on Editors if this is a Switchboard config with nDisplay devices.
+        if DeviceUnreal.csettings["force_holdout_if_ndisplay"].get_value(self.name):
+            if any(device.device_type == "nDisplay" for device in self.active_unreal_devices):
+                dp_cvars.extend([
+                    'r.PostProcessing.PropagateAlpha=1',
+                    'r.Deferred.SupportPrimitiveAlphaHoldout=1'
+                ])
+
         # Add user set dp cvars, overriding any of the forced ones.
         user_dp_cvars = self.conform_dpcvars(DeviceUnreal.csettings["dp_cvars"].get_value(self.name))
         dp_cvars = self.add_or_override_cvars(dp_cvars, user_dp_cvars)
