Switchboard: Don't pass override_device_name when constructing UI for device instance settings.

Fixes a regression from CL 35445608.

diff --git a/Engine/Plugins/VirtualProduction/Switchboard/Source/Switchboard/switchboard/settings_dialog.py b/Engine/Plugins/VirtualProduction/Switchboard/Source/Switchboard/switchboard/settings_dialog.py
--- a/Engine/Plugins/VirtualProduction/Switchboard/Source/Switchboard/switchboard/settings_dialog.py
+++ b/Engine/Plugins/VirtualProduction/Switchboard/Source/Switchboard/switchboard/settings_dialog.py
@@ -425,15 +425,22 @@ class SettingsDialog(QtCore.QObject):
     def add_plugin_settings_widgets_to_layout(
             self,
             plugin_cls: Device,
-            plugin_settings,
+            plugin_settings: list[Setting],
             layout: QtWidgets.QLayout,
             stylevariant: str = '',
-            device_name: Optional[str] = None):
+            device_name: Optional[str] = None,
+            instance_settings: list[Setting] | None = None,
+    ):
         ''' Organizes plugin_settings into categories group boxes and adds them to the given layout'''
 
+        if instance_settings is None:
+            instance_settings = []
+
+        all_settings = instance_settings + plugin_settings
+
         scategories = {}
 
-        for setting in plugin_settings:
+        for setting in all_settings:
 
             # Don't create categories for hidden properties, or you might end up with an empty category
             if not setting.show_ui:
@@ -467,7 +474,9 @@ class SettingsDialog(QtCore.QObject):
 
             # Add widgets for each setting in this category
             for setting in settings:
-                setting.create_ui(form_layout=scategory_layout, override_device_name=device_name)
+                # Instance settings don't specify override_device_name
+                override_device_name = device_name if setting in plugin_settings else None
+                setting.create_ui(form_layout=scategory_layout, override_device_name=override_device_name)
 
             # Add the category group box to the main plugin layout
             layout.addWidget(scategory_group_box)
@@ -535,16 +544,14 @@ class SettingsDialog(QtCore.QObject):
             layout = QtWidgets.QVBoxLayout()
             group_box.layout().addLayout(layout)
 
-            # Merge instance settings and class overrides to avoid repeating
-            # category group boxes under the same device instance.
-            instance_settings = settings + overrides
-
             self.add_plugin_settings_widgets_to_layout(
                 plugin_cls=plugin_cls,
-                plugin_settings=instance_settings,
+                plugin_settings=overrides,
                 layout=layout,
                 stylevariant='device_override',
-                device_name=device_name)
+                device_name=device_name,
+                instance_settings=settings,
+            )
 
         # Pull all widgets up so they do not attempt to divide space among themselves when filtered by search
         plugin_group_box.layout().addItem(
