SwitchboardListener: Move SwitchboardListenerAutolaunch.h from Internal/ to Public/.

Fixes a build break from UBT rule change in CL 38674120.

#rnx

diff --git a/Engine/Plugins/VirtualProduction/Switchboard/Source/SwitchboardCommon/Internal/SwitchboardListenerAutolaunch.h b/Engine/Plugins/VirtualProduction/Switchboard/Source/SwitchboardCommon/Internal/SwitchboardListenerAutolaunch.h
deleted file mode 100644
diff --git a/Engine/Plugins/VirtualProduction/Switchboard/Source/SwitchboardCommon/Public/SwitchboardListenerAutolaunch.h b/Engine/Plugins/VirtualProduction/Switchboard/Source/SwitchboardCommon/Public/SwitchboardListenerAutolaunch.h
new file mode 100644
--- /dev/null
+++ b/Engine/Plugins/VirtualProduction/Switchboard/Source/SwitchboardCommon/Public/SwitchboardListenerAutolaunch.h
@@ -0,0 +1,26 @@
+// Copyright Epic Games, Inc. All Rights Reserved.
+
+#pragma once
+
+#include "Containers/UnrealString.h"
+
+
+#if PLATFORM_WINDOWS
+#	define SWITCHBOARD_LISTENER_AUTOLAUNCH 1
+#else
+#	define SWITCHBOARD_LISTENER_AUTOLAUNCH 0
+#endif
+
+
+struct FLogCategoryBase;
+
+
+namespace UE::SwitchboardListener::Autolaunch
+{
+#if SWITCHBOARD_LISTENER_AUTOLAUNCH
+	SWITCHBOARDCOMMON_API FString GetInvocation(FLogCategoryBase& InLogCategory);
+	SWITCHBOARDCOMMON_API FString GetInvocationExecutable(FLogCategoryBase& InLogCategory);
+	SWITCHBOARDCOMMON_API bool SetInvocation(const FString& InNewInvocation, FLogCategoryBase& InLogCategory);
+	SWITCHBOARDCOMMON_API bool RemoveInvocation(FLogCategoryBase& InLogCategory);
+#endif // #if SWITCHBOARD_LISTENER_AUTOLAUNCH
+}
