Switchboard: Fixed not showing project settings.

Also homogenized the display of those and MU settings in the settings dialog.

#rnx

diff --git a/Engine/Plugins/VirtualProduction/Switchboard/Source/Switchboard/switchboard/settings_dialog.py b/Engine/Plugins/VirtualProduction/Switchboard/Source/Switchboard/switchboard/settings_dialog.py
--- a/Engine/Plugins/VirtualProduction/Switchboard/Source/Switchboard/switchboard/settings_dialog.py
+++ b/Engine/Plugins/VirtualProduction/Switchboard/Source/Switchboard/switchboard/settings_dialog.py
@@ -274,12 +274,19 @@ class SettingsDialog(QtCore.QObject):
     def _create_switchboard_settings(self, settings: UserSettings, config: Config):
         self.ui.switchboard_settings_group = CollapsibleGroupBox()
         self.ui.switchboard_settings_group.setTitle("Switchboard")
-        layout = QtWidgets.QFormLayout(self.ui.switchboard_settings_group)
+        layout = QtWidgets.QVBoxLayout(self.ui.switchboard_settings_group)
 
-        settings.ADDRESS.create_ui(form_layout=layout)
-        settings.TRANSPORT_PATH.create_ui(form_layout=layout)
-        config.LISTENER_EXE.create_ui(form_layout=layout)
-        config.SBLHELPER_EXE.create_ui(form_layout=layout)
+        settings_dict = {
+            "address": settings.ADDRESS,
+            "transport_path": settings.TRANSPORT_PATH,
+            "listener_exe": config.LISTENER_EXE,
+            "sblhelper_exe": config.SBLHELPER_EXE,
+        }
+
+        self.ui.switchboard_settings_sub_group = QtWidgets.QGroupBox()
+        form_layout = QtWidgets.QFormLayout(self.ui.switchboard_settings_sub_group)
+        layout.addWidget(self.ui.switchboard_settings_sub_group)
+        self._create_settings_section(settings_dict, form_layout)
 
         return self.ui.switchboard_settings_group
 
@@ -288,27 +295,26 @@ class SettingsDialog(QtCore.QObject):
         self.ui.project_settings_group.setTitle("Project Settings")
         layout = QtWidgets.QVBoxLayout(self.ui.project_settings_group)
 
-        project_settings_root = QtWidgets.QWidget()
-        form_layout = QtWidgets.QFormLayout(project_settings_root)
-        self._create_settings_section(config.basic_project_settings, form_layout)
-        layout.addWidget(project_settings_root)
-
-        # Sub settings
+        self._create_basic_settings(config, layout)
         self._create_osc_settings(config, layout)
         self._create_source_control_settings(config, layout)
 
         return self.ui.project_settings_group
 
+    def _create_basic_settings(self, config: Config, layout: QtWidgets.QFormLayout):
+        self.ui.basic_settings_group = QtWidgets.QGroupBox()
+        form_layout = QtWidgets.QFormLayout(self.ui.basic_settings_group)
+        layout.addWidget(self.ui.basic_settings_group)
+        self._create_settings_section(config.basic_project_settings, form_layout)
+
     def _create_osc_settings(self, config: Config, layout: QtWidgets.QFormLayout):
-        self.ui.osc_settings_group = QtWidgets.QGroupBox()
-        self.ui.osc_settings_group.setTitle("OSC")
+        self.ui.osc_settings_group = QtWidgets.QGroupBox("OSC")
         form_layout = QtWidgets.QFormLayout(self.ui.osc_settings_group)
         layout.addWidget(self.ui.osc_settings_group)
         self._create_settings_section(config.osc_settings, form_layout)
 
     def _create_source_control_settings(self, config: Config, layout: QtWidgets.QFormLayout):
-        self.ui.source_control_settings_group = QtWidgets.QGroupBox()
-        self.ui.source_control_settings_group.setTitle("Source Control")
+        self.ui.source_control_settings_group = QtWidgets.QGroupBox("Source Control")
         self.ui.source_control_settings_group.setCheckable(True)
         self.ui.source_control_settings_group.setStyleSheet(
             'QGroupBox::indicator:checked:hover {image: url(:icons/images/check_box_checked_hovered.png);}'
@@ -337,9 +343,13 @@ class SettingsDialog(QtCore.QObject):
     def _create_multi_user_server_settings(self, config: Config):
         self.ui.multi_user_settings = CollapsibleGroupBox()
         self.ui.multi_user_settings.setTitle("Multi User Server")
-        form_layout = QtWidgets.QFormLayout(self.ui.multi_user_settings)
+        layout = QtWidgets.QVBoxLayout(self.ui.multi_user_settings)
 
+        self.ui.multi_user_settings_sub_group = QtWidgets.QGroupBox()
+        form_layout = QtWidgets.QFormLayout(self.ui.multi_user_settings_sub_group)
+        layout.addWidget(self.ui.multi_user_settings_sub_group)
         self._create_settings_section(config.mu_settings, form_layout)
+
         return self.ui.multi_user_settings
 
     def _create_settings_section(self, dict: Dict[str, Setting], layout: QtWidgets.QFormLayout):
