Switchboard: Exclude AdditionalPathsToSync from code change consideration.

diff --git a/Engine/Plugins/VirtualProduction/Switchboard/Source/Switchboard/sbl_helper.py b/Engine/Plugins/VirtualProduction/Switchboard/Source/Switchboard/sbl_helper.py
--- a/Engine/Plugins/VirtualProduction/Switchboard/Source/Switchboard/sbl_helper.py
+++ b/Engine/Plugins/VirtualProduction/Switchboard/Source/Switchboard/sbl_helper.py
@@ -675,7 +675,9 @@ class SbListenerHelper:
 
         if not self.use_ugs:
             logging.info('Determining CompatibleChangelist')
-            buildver_info = self.get_buildver_info(engine_paths, project_paths)
+
+            # [:1] slice because we exclude AdditionalPathsToSync here
+            buildver_info = self.get_buildver_info(engine_paths, project_paths[:1])
 
             if self.sync_engine_cl:
                 # Remove Build.version if present, so we can write our own.
