Switchboard: Add support for EUnrealEngineObjectUE5Version::METADATA_SERIALIZATION_OFFSET to uassetparser.py.

Corresponds to engine serialization changes in CL 40093672.

diff --git a/Engine/Plugins/VirtualProduction/Switchboard/Source/Switchboard/switchboard/devices/unreal/uassetparser.py b/Engine/Plugins/VirtualProduction/Switchboard/Source/Switchboard/switchboard/devices/unreal/uassetparser.py
--- a/Engine/Plugins/VirtualProduction/Switchboard/Source/Switchboard/switchboard/devices/unreal/uassetparser.py
+++ b/Engine/Plugins/VirtualProduction/Switchboard/Source/Switchboard/switchboard/devices/unreal/uassetparser.py
@@ -27,7 +27,7 @@ class AssetRegistryData(object):
         self.ObjectPath = ''
         self.ObjectClassName = ''
         self.tags = {}
-        
+
 class FStringError(Exception):
     pass
 
@@ -54,11 +54,11 @@ class UassetParser(object):
                 for k,v in assetdata.tags.items():
                     print('Tag {}: {}'.format(k,v))
     '''
-    
+
     _names = None
     _aregdata = None
     _thumbnailCache = None
-    
+
     def __init__(self, fileObj: BinaryIO, allowUnversioned=True):
         ''' Parses `fileObj` and reads the summary, which has the most basic information about the asset. '''
 
@@ -74,13 +74,13 @@ class UassetParser(object):
     @property
     def thumbnailCache(self):
         ''' Lazy read of the thumbnailCache '''
-        
+
         if self._thumbnailCache is not None:
             return self._thumbnailCache
 
         self._thumbnailCache = self._readAssetDataFromThumbnailCache()
         return self._thumbnailCache
-        
+
     @property
     def names(self):
         ''' Lazy read of the "names" of the asset '''
@@ -89,44 +89,44 @@ class UassetParser(object):
 
         self._names = self._readNames()
         return self._names
-    
+
     @property
     def aregdata(self):
         ''' Lazy read of the asset registry data '''
         if self._aregdata is None:
             self._aregdata = self._readAssetRegistryData()
-            
+
         return self._aregdata
-    
+
     def readFString(self):
         ''' Reads an FString from the current file position '''
         sz, = struct.unpack('<i',self.f.read(4))
-        
+
         if sz == 0:
             return ''
-        
+
         LoadUCS2Char = False
-        
+
         if sz < 0:
             LoadUCS2Char = True
             sz = -sz
-        
+
         if LoadUCS2Char:
             sz *= 2
-        
+
         msg, nullTerminator = struct.unpack('%dsB' % (sz-1), self.f.read(sz))
-        
+
         if LoadUCS2Char:
             msg = msg[:-1].decode('utf-16')
         else:
             msg = msg.decode('utf-8')
-            
+
         return msg
 
     def skipBytes(self,n):
         ''' Skips the specified number of bytes from current file position '''
         self.f.seek(self.f.tell() + n)
-        
+
     def skipTArray(self, typeSz, maxElements=100000):
         ''' Skips a TArray that holds the specified type size '''
         n = self.readInt32()
@@ -162,52 +162,52 @@ class UassetParser(object):
 
     def readTArray(self, elementreaderfn, maxElements=100000):
         ''' Reads TArray of any type.
-        
+
         Args:
             elementreaderfn: A function that reads the element at current file offset.
         '''
         arr = []
         n = self.readInt32()
-        
+
         if not (0 <= n <= maxElements):
             raise ValueError("Invalid TArray.Num() == {}".format(n))
 
         for _ in range(n):
             arr.append(elementreaderfn())
-        
+
         return arr
-    
+
     def _readNames(self):
         ''' Read the names of the asset '''
-    
+
         if self.asummary.NameCount <= 0:
             return []
-        
+
         fos = self.asummary.NameOffset
-        
+
         if (fos <= 0) or (fos > self.PackageFileSize):
             return []
-        
+
         self.f.seek(fos)
-            
+
         names = []
-        
+
         for _ in range(self.asummary.NameCount):
             name = self.readFString()
             self.skipBytes(4) # precalculated hashes
             names.append(name)
-            
+
         return names
-        
+
     def _readAssetRegistryData(self):
         ''' Read the asset registry data. Find code in PackageUtilities.cpp and PackageReader.cpp
         '''
 
         fos = self.asummary.AssetRegistryDataOffset
-        
+
         if (fos <= 0) or (fos > self.PackageFileSize):
             return []
-        
+
         self.f.seek(fos)
 
         DependencyDataOffset = self.readInt64()
@@ -218,19 +218,19 @@ class UassetParser(object):
 
         if nAssets < 0:
             raise ValueError('Invalid number of {} AssetRegistryDatas'.format(nAssets))
-        
+
         assets = []
-        
+
         for _ in range(nAssets):
             asset = AssetRegistryData()
             asset.ObjectPath = self.readFString()
             asset.ObjectClassName = self.readFString()
-            
+
             assets.append(asset)
-            
+
             nTags = self.readInt32()
             asset.tags = {}
-            
+
             for _ in range(nTags):
                 try:
                     key = self.readFString()
@@ -238,7 +238,7 @@ class UassetParser(object):
                     asset.tags[key] = val
                 except FStringError:
                     return assets
-             
+
         return assets
 
     def checkFileOffset(self, fos):
@@ -274,24 +274,26 @@ class UassetParser(object):
         ''' Read the package summary of the asset
         '''
         self.f.seek(0)
-        
+
         s = UassetSummary()
-        
+
         s.Tag = self.readUInt32()
 
         if s.Tag != 0x9e2a83c1:
             raise ValueError('Not a valid uasset!')
 
         s.LegacyFileVersion = self.readInt32()
-        
+
         if s.LegacyFileVersion not in [-7,-8]:
             raise ValueError('LegacyFileVersion {} is non handled by this parser'.format(s.LegacyFileVersion))
 
         s.LegacyUE3Version = self.readInt32()
         s.FileVersionUE4 = self.readInt32()
 
-        if s.LegacyFileVersion == -8:
+        if s.LegacyFileVersion <= -8:
             s.FileVersionUE5 = self.readInt32()
+        else:
+            s.FileVersionUE5 = 0
 
         s.FileVersionLicenseeUE4 = self.readUInt32()
 
@@ -305,17 +307,22 @@ class UassetParser(object):
             s.SoftObjectPathsOffset = self.readInt32()
 
         s.LocalizationId = self.readFString()
-        
+
         def remainingBytesInHeader():
             ''' Convenience local function to calculate remaning bytes in header. Useful for sanity checks during parsing '''
             return s.TotalHeaderSize - self.f.tell() + 1
 
-        s.GatherableTextDataCount, s.GatherableTextDataOffset, s.ExportCount, s.ExportOffset, \
-        s.ImportCount, s.ImportOffset, s.DependsOffset, s.SoftPackageReferencesCount, \
-        s.SoftPackageReferencesOffset, s.SearchableNamesOffset, s.ThumbnailTableOffset, s.Guid, \
-        s.PersistentGuid \
-        = struct.unpack('<11i16s16s', self.f.read(11*4 + 16 + 16))
-        
+        s.GatherableTextDataCount, s.GatherableTextDataOffset, s.ExportCount, \
+        s.ExportOffset, s.ImportCount, s.ImportOffset \
+            = struct.unpack('<6i', self.f.read(6*4))
+
+        if s.FileVersionUE5 >= 1014:  # EUnrealEngineObjectUE5Version::METADATA_SERIALIZATION_OFFSET
+            s.MetaDataOffset = self.readInt32()
+
+        s.DependsOffset, s.SoftPackageReferencesCount, s.SoftPackageReferencesOffset, \
+        s.SearchableNamesOffset, s.ThumbnailTableOffset, s.Guid, s.PersistentGuid \
+            = struct.unpack('<5i16s16s', self.f.read(5*4 + 16 + 16))
+
         self.checkFileOffset(s.GatherableTextDataOffset)
         self.checkFileOffset(s.ExportOffset)
         self.checkFileOffset(s.ImportOffset)
@@ -325,15 +332,15 @@ class UassetParser(object):
         self.checkFileOffset(s.ThumbnailTableOffset)
 
         s.Generations = self.readTArray(lambda : self.f.read(8), maxElements=remainingBytesInHeader()/20)
-        
+
         s.SavedByEngineVersionMajor, s.SavedByEngineVersionMinor, s.SavedByEngineVersionPatch, \
         s.SavedByEngineVersionChangelist = struct.unpack('<HHHI', self.f.read(2*3+4))
         s.SavedByEngineVersionName = self.readFString()
-        
+
         s.CompatibleEngineVersionMajor, s.CompatibleEngineVersionMinor, s.CompatibleEngineVersionPatch, \
         s.CompatibleEngineVersionChangelist = struct.unpack('<HHHI', self.f.read(2*3+4))
         s.CompatibleEngineVersionName = self.readFString()
-        
+
         self.checkAssetVersion(s.SavedByEngineVersionMajor, s.SavedByEngineVersionMinor, s.SavedByEngineVersionPatch)
 
         s.CompressionFlags = self.readUInt32()
@@ -347,7 +354,7 @@ class UassetParser(object):
         s.PackageSource = self.readUInt32()
 
         s.AdditionalPackagesToCook = self.readTArray(self.readFString, maxElements=remainingBytesInHeader())
-            
+
         s.AssetRegistryDataOffset = self.readInt32()
         s.BulkDataStartOffset = self.readInt64()
 
@@ -355,30 +362,30 @@ class UassetParser(object):
         self.checkFileOffset(s.BulkDataStartOffset)
 
         return s
-        
+
     def _readAssetDataFromThumbnailCache(self):
         ''' Read the asset data from the thumbnail cache.
         '''
         fos = self.asummary.ThumbnailTableOffset
-        
+
         if (fos <= 0) or (fos > self.PackageFileSize):
             return []
-            
+
         self.f.seek(fos)
-        
+
         ObjectCount = self.readInt32()
-        
+
         AssetDataList = []
-        
+
         for _ in range(ObjectCount):
             assetData = AssetData()
-        
+
             assetData.AssetClassName = self.readFString()
             assetData.ObjectPathWithoutPackageName = self.readFString()
             assetData.FileOffset = self.readInt32()
 
             AssetDataList.append(assetData)
-            
+
         return AssetDataList    
 
 def printAssetData(aparser, bAssetRegistry, bTags, bNames, bThumbnailCache):
@@ -431,4 +438,3 @@ if __name__ == '__main__':
     with open(fpath, 'rb') as file:
         aparser = UassetParser(file)
         printAssetData(aparser, bAssetRegistry, bTags, bNames, bThumbnailCache)
-    
\ No newline at end of file
