Switchboard: Build indicator fixes, code cleanup, and documentation.

Previously, the `CompatibleChangelist` value from `UnrealEditor.version` was regarded as the source of truth for whether the build was up-to-date, which was inadequate. At the same time, the code to determine the expected value was was more complicated than it needed to be.

Now, we compare both `UnrealEditor.version` and `$(ProjectName)Editor.target` against `Build.version`, and for each file, we compare both the `Changelist` and `CompatibleChangelist` values. Switchboard does not need the code changelist determination logic, because that determination is already made at sync time by `sbl_helper.py`, and the computed value cached in `CompatibleChangelist`.

Also refactors `UnrealPluginManager` to expose `AdditionalPluginDirectories` from the .uproject (which will likely be required for some sync indicator fixes in the near future), and fixes some linter warnings where the built-in `id()` function was being shadowed.

diff --git a/Engine/Plugins/VirtualProduction/Switchboard/Source/Switchboard/docs/SyncAndVersioning.md b/Engine/Plugins/VirtualProduction/Switchboard/Source/Switchboard/docs/SyncAndVersioning.md
new file mode 100644
--- /dev/null
+++ b/Engine/Plugins/VirtualProduction/Switchboard/Source/Switchboard/docs/SyncAndVersioning.md
@@ -0,0 +1,134 @@
+# Introduction
+When building Unreal Engine from source code, the use of the UnrealGameSync utility typically plays an essential role in the correct versioning of Unreal projects, particularly when it comes to package saving and loading compatibility determination.
+
+However, while UGS assumes a single revision across the entire workspace, Switchboard enables users to choose which revision to sync their project and engine to separately. This capability has proven invaluable in virtual production workflows, enabling users to ensure stability by freezing the bulk of the code at a known-good version, while still permitting them to iterate on content and disseminate project updates in a fast-paced stage environment.
+
+With a handful of notable exceptions, Switchboard (and its sync script `sbl_helper.py`) endeavors to be consistent with how UnrealGameSync handles versioning. This document is intended to provide a quick overview of some of the salient details, caveats, and rationale for places where we've diverged in behavior.
+
+> ⚠️ **Warning**: For reasons which will be elaborated upon throughout this document, UnrealGameSync should not be used in a workspace managed by Switchboard, and vice versa. UnrealGameSync relies on the assumption that its cached state remains valid, including but not limited to the last changelist synced, and sync filters being consistently applied.
+>
+> One possible failure mode is that syncing with Switchboard can cause files that would otherwise be excluded via UGS sync filters to be re-added to the workspace, but a subsequent UGS sync using filters would neither update nor remove them, leading to version skew.
+
+
+# Background: UnrealGameSync
+This section describes how some of the relevant systems work using standard UGS workflows.
+
+### `Build.version`
+Stored alongside the engine is a JSON file, `Engine/Build/Build.version`, which contains several fields describing aspects of the Unreal Engine version.
+
+Here's an example of how the file is stored in version control in the UE5 Main stream/branch at the time of writing:
+```json
+{
+	"MajorVersion": 5,
+	"MinorVersion": 6,
+	"PatchVersion": 0,
+	"Changelist": 0,
+	"CompatibleChangelist": 0,
+	"IsLicenseeVersion": 0,
+	"IsPromotedBuild": 0,
+	"BranchName": "UE5"
+}
+```
+
+Whereas this is how it appears in UE 5.5.3 installed via the Epic Games Launcher:
+```json
+{
+	"MajorVersion": 5,
+	"MinorVersion": 5,
+	"PatchVersion": 3,
+	"Changelist": 39772772,
+	"CompatibleChangelist": 37670630,
+	"IsLicenseeVersion": 0,
+	"IsPromotedBuild": 1,
+	"BranchName": "++UE5+Release-5.5"
+}
+```
+
+This file is treated as an input to the build process, and many of the values become preprocessor definitions which are "baked into" the compiled binaries.
+
+While some fields are manually updated by editing the file as it's stored in Perforce as one might expect, other fields are automatically rewritten by UnrealGameSync (or Switchboard) on the local machine as part of every sync operation performed. Here's a breakdown:
+
+  * `MajorVersion`, `MinorVersion`, and `PatchVersion` are only updated manually by Epic at certain points during our development and release cycle.
+
+  * `Changelist` is always rewritten following sync operations, and reflects whichever revision the user has selected to sync to. **Notably, this value is used internally by the Concert (Multi-User) code to attempt to ensure that everyone joining the session is doing so with the same base revision of the content, and is stored as part of the session.**
+
+  * `CompatibleChangelist` is relevant when reading/writing uasset files, and used for determining whether a saved asset is compatible with the running engine code attempting to load it. The handling of this field is the most complex, and differs significantly depending on whether the file stored in Perforce has a value of 0:
+
+    * If the value stored in Perforce is 0, which is typical of a mainline development stream/branch, then it will be updated on sync to reflect the most recent "code changelist"* less than or equal to the (potentially content-only) changelist the user has selected to sync to. Classification of a changelist as affecting code is based on the presence of certain modified file extensions (e.g. .h/.cpp).
+      > **: Refer to the `p4_latest_code_change` function in `p4_utils.py` for details.*
+
+    * If this value is NOT 0, UnrealGameSync will NOT update it automatically, and will instead use it as-is. This is indicative of a release stream/branch which is in hotfix/maintenance mode. The value will have been manually updated to match the changelist that corresponds to the initial (5.x.0) release, and reflects the manual scrutiny Epic applies internally when vetting changes for inclusion in hotfix releases in order to maintain serialization and ABI compatibility.
+
+  * `IsLicenseeVersion` is a boolean flag that's updated on sync to reflect the presence or absence of specific files in the engine source tree that are only available to Epic employees. This in turn sets a single bit causing `CompatibleChangelist` to be interpreted differently. (TODO: Elaborate on this.)
+
+  * `IsPromotedBuild` is a boolean flag that's only set as part of Epic's internal release build process when producing official binary installations for the Epic Games Launcher.
+
+  * `BranchName` is always updated on sync to reflect the Perforce stream name from which the workspace was synced, with plus characters substituted for forward slashes (e.g. `++UE5+Main` or `++UE5+Release-5.5`).
+
+### `UnrealEditor.version`
+Produced by UnrealBuildTool as part of a successful build is an analogous output artifact, written alongside the compiled executable, e.g. `Engine/Binaries/Win64/UnrealEditor.version`. This file is largely a duplicate of the `Build.version` file read as input at build time, with the addition of a `BuildId` string identifier appended to the end.
+
+```json
+{
+	"...": "...",
+	"BranchName": "++UE5+Main",
+	"BuildId": "f28a59da-6a9b-46e1-93e3-d5843d7c446d"
+}
+```
+
+The same `BuildId` value is also written to the `.modules` files alongside all modular DLLs, and intended to ensure all dynamically loaded modules have been built from the same code, and are therefore ABI-compatible. The exact nature of the value depends on whether a source build or an official EGL binary build is being used:
+
+  * For source builds, it's set to a new randomly-generated UUID on each build.
+
+  * For launcher builds, this value is overridden* to the same `CompatibleChangelist` number discussed above, corresponding to the initial minor (5.x.0) version release, and with the same rationale: a great deal of manual scrutiny is applied to ensure ABI compatibility between hotfixes.
+    > **: Refer to the `FinalizeInstalledBuild` UnrealAutomationTool script for details.*
+
+If you've ever been notified that modules are out of date and need to be rebuilt, a mismatch in the `BuildId` field is probably the reason (for example, syncing and build a project that does not use plugin X, then subsequently attempting to load a different project which would otherwise attempt to use existing stale binaries for plugin X).
+
+### `$(ProjectName)Editor.target`
+A related build output, known as a target receipt, is produced for the active project's editor target upon each successful build. The target receipt contains, among (many) other things, another copy of the same fields as `UnrealEditor.version`.
+
+Importantly, however, while `UnrealEditor.version` is updated when building *any* project, target receipts are specific to a single project's editor target.
+
+### `UnrealGameSync.ini`
+UnrealGameSync coalesces configuration from a variety of possible .ini file locations distributed throughout the engine and project directory hierarchies*, but the two most commonly used locations are the stream-wide `/Engine/Programs/UnrealGameSync/UnrealGameSync.ini`, and the per-project `$(ProjectDir)/Build/UnrealGameSync.ini`.
+> **: Refer to the `get_depot_config_paths` function in `ugs_utils.py` for details.*
+
+By modifying these files, it's possible to influence the behavior of UGS both stream-wide, as well in the context of only the active project.
+
+One such configuration directive which we have added support for in Switchboard is `AdditionalPathsToSync`. This allows users to configure paths that are neither part of the engine nor part of the project, but are synced simultaneously. We have found this valuable in virtual production workflows to have a content plugin that serves as a shared asset library used across multiple projects. (See also: "Additional Plugin Directories" functionality in the Plugin Browser.)
+
+
+# Switchboard: Engine vs. Project Changelists
+As mentioned previously, Switchboard supports workflows where the project may be synced to a different revision than the engine. Specifically:
+
+  * Switchboard considers the "Engine" to consist of any files beneath the `/Engine` directory, combined with any loose files at the root of the workspace (e.g. `/GenerateProjectFiles.bat`). These are the paths that are enumerated for changes in order to populate the "Engine" changelist selection dropdown.
+
+  * Switchboard considers the "Project" to consist of any files located in the directory containing the uproject file, as well as any paths referenced via the UnrealGameSync `AdditionalPathsToSync` directive. We consider these additional paths to be versioned "with the project" because it is frequently useful to sync external plugins to a later version than the engine for the same reasons (whereas plugins stored within the engine are already shared between projects). All of these paths are enumerated for changes when populating the "Project" changelist dropdown.
+
+### Differences With UGS
+Because of this, some versioning concepts designed with UnrealGameSync in mind do not map entirely cleanly to Switchboard:
+
+  * Switchboard sets the `Changelist` field to whichever is the more recent of the selected engine/project changelists.
+
+  * Switchboard sets the `CompatibleChangelist` field to the most recent code change in either the engine or the .uproject directory, but does NOT consider code changes in `AdditionalPathsToSync`.
+
+  * Switchboard ignores any `CompatibleChangelist` value stored in Perforce, as hotfix release compatibility is not likely to be relevant to stage workflows, nor is it easy to guarantee when dealing with arbitrary project code. (This behavior can be overridden by passing `--ugs-versioning` to `sbl_helper.py`.)
+
+
+# Switchboard: Sync and Build Indicators
+Switchboard has per-device indicators of sync and build status, driven by various aspects of the device's workspace.
+
+### Sync
+The indication of what changelists a device is synced to is based on invoking `p4 cstat` for the engine or project paths with the `#have` revision specifier. The changelist returned by the command is compared against the one selected from the corresponding Switchboard combo box. If there's a mismatch, the corresponding changelist indicator and the sync button will both be highlighted, indicating the device needs to be synced to the desired changelist.
+
+### Build
+Similarly, Switchboard will indicate if a device has not yet built the changelist it's currently synced to. For this, Switchboard relies on a few relevant facts:
+
+  * The `Build.version` file will be updated on every sync performed, with both the changelist selected for sync, as well as the corresponding code changelist (`Changelist` and `CompatibleChangelist`, respectively).
+
+  * The `UnrealEditor.version` file, along with any shared engine modules, will be updated upon building *any* project.
+
+  * The `$(ProjectName)Editor.target` file will be updated to match on each build of a *specific* project.
+
+By comparing these fields across all three files, we can be relatively confident that the relevant build outputs are up-to-date with the last sync performed in the workspace. Any mismatch will cause the build button to be highlighted, and the tooltip will contain the relevant details.
diff --git a/Engine/Plugins/VirtualProduction/Switchboard/Source/Switchboard/switchboard/devices/device_base.py b/Engine/Plugins/VirtualProduction/Switchboard/Source/Switchboard/switchboard/devices/device_base.py
--- a/Engine/Plugins/VirtualProduction/Switchboard/Source/Switchboard/switchboard/devices/device_base.py
+++ b/Engine/Plugins/VirtualProduction/Switchboard/Source/Switchboard/switchboard/devices/device_base.py
@@ -55,7 +55,7 @@ class DeviceQtHandler(QtCore.QObject):
     signal_device_client_disconnected = QtCore.Signal(object)
     signal_device_project_changelist_changed = QtCore.Signal(object)
     signal_device_engine_changelist_changed = QtCore.Signal(object)
-    signal_device_built_engine_changelist_changed = QtCore.Signal(object)
+    signal_device_built_changelist_changed = QtCore.Signal(object)
     signal_device_sync_failed = QtCore.Signal(object)
     signal_device_is_recording_device_changed = QtCore.Signal(object, int)
 
@@ -105,7 +105,6 @@ class Device(QtCore.QObject):
 
         self._project_changelist: Optional[int] = None
         self._engine_changelist: Optional[int] = None
-        self._built_engine_changelist: Optional[int] = None
 
         self._status = DeviceStatus.DISCONNECTED
 
@@ -281,17 +280,6 @@ class Device(QtCore.QObject):
         self.device_qt_handler.signal_device_engine_changelist_changed.emit(
             self)
 
-    @property
-    def built_engine_changelist(self):
-        return self._built_engine_changelist
-
-    @built_engine_changelist.setter
-    def built_engine_changelist(self, value: Optional[int]):
-        self._built_engine_changelist = value
-        self.device_qt_handler.signal_device_built_engine_changelist_changed.emit(
-            self
-        )
-
     @property
     def device_osc_port(self) -> int:
         raise NotImplementedError()
diff --git a/Engine/Plugins/VirtualProduction/Switchboard/Source/Switchboard/switchboard/devices/unreal/plugin_unreal.py b/Engine/Plugins/VirtualProduction/Switchboard/Source/Switchboard/switchboard/devices/unreal/plugin_unreal.py
--- a/Engine/Plugins/VirtualProduction/Switchboard/Source/Switchboard/switchboard/devices/unreal/plugin_unreal.py
+++ b/Engine/Plugins/VirtualProduction/Switchboard/Source/Switchboard/switchboard/devices/unreal/plugin_unreal.py
@@ -43,7 +43,6 @@ from switchboard.devices.unreal.uassetparser import UassetParser
 from switchboard.listener_client import ListenerClient
 from switchboard.switchboard_logging import LOGGER
 from switchboard.tools.insights_launcher import InsightsLauncher
-from switchboard.util import p4_changelist_inspection
 from . import version_helpers
 from .listener_watcher import ListenerWatcher
 from .redeploy_dialog import RedeployListenerDialog
@@ -487,11 +486,11 @@ class MediaProfileSetting(LiveLinkPresetSetting):
 
 class SyncCategoryOption:
     ''' Used to represent a UGS sync filter category within a Setting. '''
-    def __init__(self, id: Union[str, uuid.UUID], name: str):
-        if isinstance(id, str):
-            id = uuid.UUID(id)
+    def __init__(self, catid: Union[str, uuid.UUID], name: str):
+        if isinstance(catid, str):
+            catid = uuid.UUID(catid)
 
-        self.id = id
+        self.catid = catid
         self.name = name
 
     def __str__(self) -> str:
@@ -501,7 +500,7 @@ class SyncCategoryOption:
     def __eq__(self, other: object):
         ''' Two category options compare equal if their UUIDs are equal. '''
         if isinstance(other, SyncCategoryOption):
-            return self.id == other.id
+            return self.catid == other.catid
         else:
             return False
 
@@ -761,7 +760,7 @@ class DeviceUnreal(Device):
             tool_tip='UnrealGameSync filter categories to include during sync',
             category="Source Control Settings",
         ).with_get_json_override_fn(
-            lambda val_list: [(str(x.id), x.name) for x in val_list]
+            lambda val_list: [(str(x.catid), x.name) for x in val_list]
         ).with_config_set_override_fn(
             lambda val_list: [
                 SyncCategoryOption(uuid.UUID(x[0]), x[1])
@@ -786,6 +785,16 @@ class DeviceUnreal(Device):
                 'this Switchboard configuration \n'),
             category="UE Settings",
         ),
+        'editor_target_override': StringSetting(
+            attr_name='editor_target_override',
+            nice_name='Editor Target Name Override',
+            value='',
+            placeholder_text='$(ProjectName)Editor',
+            tool_tip=(
+                "The name of your project's editor target. Typically $(ProjectName)Editor (which is inferred "
+                "automatically). This override only needs to be set if deviating from the standard convention."),
+            category='UE Settings',
+        )
     }
 
     if ENABLE_UGS_SUPPORT:
@@ -1073,10 +1082,33 @@ class DeviceUnreal(Device):
         self.auto_connect = False
 
         self.runtime_str = ""
+
+        # The CLs returned via p4 cstat query, or optimistically
+        # updated with the "inflight" values upon successful sync
+        # (Defined in the base class, documented here for context)
+        #self._project_changelist: Optional[int] = None
+        #self._engine_changelist: Optional[int] = None
+
+        # The CLs for which outstanding p4 syncs have been queued
         self.inflight_project_cl = None
         self.inflight_engine_cl = None
         self.inflight_sync_hash: Optional[str] = None
 
+        # The fields read from Build.version JSON
+        # (Written during sync by UGS/sbl_helper.py)
+        self._synced_changelist: Optional[int] = None
+        self._synced_compatible_cl: Optional[int] = None
+
+        # The fields read from UnrealEditor.version JSON
+        # (Written by UBT following a successful build of ANY project)
+        self._built_eng_changelist: Optional[int] = None
+        self._built_eng_compatible_cl: Optional[int] = None
+
+        # The fields read from (Project)Editor.target JSON
+        # (Written by UBT following a successful build of the active project)
+        self._built_proj_changelist: Optional[int] = None
+        self._built_proj_compatible_cl: Optional[int] = None
+
         listener_qt_handler = self.unreal_client.listener_qt_handler
         listener_qt_handler.listener_connecting.connect(
             super().connecting_listener, QtCore.Qt.QueuedConnection)
@@ -1460,7 +1492,7 @@ class DeviceUnreal(Device):
                 self._request_engine_changelist_number()
         else:
             self.widget.engine_changelist_label.hide()
-            self.widget.update_build_info(synched_cl=self.engine_changelist, built_cl=self.built_engine_changelist)
+            self.widget.update_build_info(self)
 
     def on_setting_exclude_from_build_changed(self, exclude_from_build):
         self.widget.update_exclude_from_build(exclude_from_build, not self.is_disconnected)
@@ -1500,17 +1532,36 @@ class DeviceUnreal(Device):
         _, msg = message_protocol.create_copy_file_from_listener_message(roles_file_path)
         self.unreal_client.send_message(msg)
 
-    def _request_unreal_editor_version_file(self):
+    def _request_version_files(self):
         '''
-        Requests the Engine/Binaries/[platform]/UnrealEditor.version file.
-        On success _on_receive_editor_version is called with the file content.
+        Requests the Build.version, UnrealEditor.version and $(Proj)Editor.target files.
+        On success, _on_receive_version_file is called with the file content.
         '''
-        engine_path = CONFIG.ENGINE_DIR.get_value(self.name).replace('"', '')
+        engine_dir = CONFIG.ENGINE_DIR.get_value(self.name)
+        uproj_path = CONFIG.UPROJECT_PATH.get_value(self.name)
+        proj_dir = os.path.dirname(uproj_path)
         platform_dir = self.platform_binary_directory
-        roles_file_path = os.path.join(os.path.dirname(engine_path), "Engine", "Binaries", platform_dir, "UnrealEditor.version")
-        _, msg = message_protocol.create_copy_file_from_listener_message(roles_file_path)
+
+        proj_name, _ = os.path.splitext(os.path.basename(uproj_path))
+        editor_target = DeviceUnreal.csettings['editor_target_override'].get_value()
+        if not editor_target:
+            editor_target = f'{proj_name}Editor'
+
+        build_version_path = os.path.join(engine_dir, 'Build', 'Build.version')
+        _, msg = message_protocol.create_copy_file_from_listener_message(build_version_path)
+        self.unreal_client.send_message(msg)
+
+        editor_version_path = os.path.join(engine_dir, 'Binaries', platform_dir,
+                                           'UnrealEditor.version')
+        _, msg = message_protocol.create_copy_file_from_listener_message(editor_version_path)
         self.unreal_client.send_message(msg)
 
+        proj_target_path = os.path.join(proj_dir, 'Binaries', platform_dir,
+                                        f'{editor_target}.target')
+        _, msg = message_protocol.create_copy_file_from_listener_message(proj_target_path)
+        self.unreal_client.send_message(msg)
+
+
     def _request_project_changelist_number(self):
         if not CONFIG.P4_ENABLED.get_value():
             return
@@ -1747,8 +1798,8 @@ class DeviceUnreal(Device):
             incl_categories: list[SyncCategoryOption] = self.csettings[
                 'included_sync_categories'].get_value()
             if incl_categories:
-                incl_categories.sort(key=lambda cat: cat.id)
-                include_str = ','.join(str(cat.id) for cat in incl_categories)
+                incl_categories.sort(key=lambda cat: cat.catid)
+                include_str = ','.join(str(cat.catid) for cat in incl_categories)
                 sync_args += f' --include-categories="{include_str}"'
 
             if custom_view_str:
@@ -2587,7 +2638,7 @@ class DeviceUnreal(Device):
 
             # If you build and sync the engine, update its CL
             if CONFIG.ENGINE_SYNC_METHOD.get_value() == EngineSyncMethod.Build_Engine.value:
-
+                self._request_version_files()
                 if self.inflight_engine_cl:
                     self.engine_changelist = self.inflight_engine_cl
                     self.inflight_engine_cl = None
@@ -2631,7 +2682,7 @@ class DeviceUnreal(Device):
                 self._request_project_changelist_number()
                 if CONFIG.ENGINE_SYNC_METHOD.get_value() == EngineSyncMethod.Build_Engine.value:
                     self._request_engine_changelist_number()
-                    self._request_unreal_editor_version_file()
+                    self._request_version_files()
 
         elif UnrealJobs.CstatPrefix.value in program_name:
             output = get_stdout_str()
@@ -2697,9 +2748,9 @@ class DeviceUnreal(Device):
     def on_file_received(self, source_path, content):
         if source_path.endswith(DeviceUnreal.csettings["roles_filename"].get_value(self.name)):
             self._on_receive_roles(content)
-        elif "UnrealEditor" in source_path and source_path.endswith(".version"):
-            self._on_receive_editor_version(content)
-        
+        elif source_path.endswith('.version') or source_path.endswith('.target'):
+            self._on_receive_version_file(source_path, content)
+
     def _on_receive_roles(self, content):
         decoded_content = base64.b64decode(content).decode()
         tags = parse_unreal_tag_file(decoded_content.splitlines())
@@ -2716,15 +2767,31 @@ class DeviceUnreal(Device):
                 f"{self.name}: Please change the roles for this device in "
                 "the settings or in the unreal project settings!")
 
-    def _on_receive_editor_version(self, content):
+    def _on_receive_version_file(self, source_path: str, content):
         '''
-        Receives the Engine/Binaries/[platform]/UnrealEditor.version file _request_unreal_editor_version_file
+        Handles .version/.target files (see _request_version_files)
         '''
         decoded_content = base64.b64decode(content).decode()
         data = json.loads(decoded_content)
-        compatible_cl = data.get("CompatibleChangelist", None)
+
+        if source_path.endswith('.target'):
+            data = data.get('Version')
+
+        changelist = data.get('Changelist')
+        compatible_cl = data.get('CompatibleChangelist')
+
         for device in self.devices_sharing_workspace():
-            device.built_engine_changelist = compatible_cl
+            if source_path.endswith('Build.version'):
+                device._synced_changelist = changelist
+                device._synced_compatible_cl = compatible_cl
+            elif source_path.endswith('UnrealEditor.version'):
+                device._built_eng_changelist = changelist
+                device._built_eng_compatible_cl = compatible_cl
+            elif source_path.endswith('.target'):
+                device._built_proj_changelist = changelist
+                device._built_proj_compatible_cl = compatible_cl
+
+        self.device_qt_handler.signal_device_built_changelist_changed.emit(self)
 
     def on_file_receive_failed(self, source_path, error):
         if source_path.endswith(DeviceUnreal.csettings["roles_filename"].get_value(self.name)):
@@ -2733,7 +2800,8 @@ class DeviceUnreal(Device):
                 LOGGER.error(
                     f"{self.name}: Error receiving role file from listener and "
                     f"device claims to have these roles: {' | '.join(roles)}")
-                LOGGER.error(f"Error: {error}")
+
+        LOGGER.error(f"{self.name}: Error retrieving file {source_path}: {error}")
 
     def on_listener_programstdout(self, message):
         ''' Handles updates to stdout of programs
@@ -2910,7 +2978,7 @@ class DeviceUnreal(Device):
 
         if CONFIG.ENGINE_SYNC_METHOD.get_value() == EngineSyncMethod.Build_Engine.value:
             self._request_engine_changelist_number()
-            self._request_unreal_editor_version_file()
+            self._request_version_files()
 
     def transport_paths(self, device_recording):
         '''
@@ -3657,20 +3725,26 @@ class DeviceWidgetUnreal(DeviceWidget):
 
         self.project_changelist_label.show()
 
-    def update_engine_changelist(self, required_cl: str, synched_cl: str, built__cl: str):
+    def update_engine_changelist(
+        self,
+        required_cl: str,
+        device: DeviceUnreal,
+    ):
         if not CONFIG.ENGINE_SYNC_METHOD.get_value() == EngineSyncMethod.Build_Engine.value:
             return
 
-        self.engine_changelist_label.setText(f'E: {synched_cl}')
+        synced_cl = device.engine_changelist
+
+        self.engine_changelist_label.setText(f'E: {synced_cl}')
         if not self.exclude_from_build:
             self.engine_changelist_label.show()
 
         self.sync_button.show()
         self.build_button.show()
 
-        is_synched = required_cl is None or required_cl == synched_cl
+        is_synched = required_cl is None or required_cl == synced_cl
         self._set_engine_changelist_is_synched(is_synched)
-        self.update_build_info(synched_cl=synched_cl, built_cl=built__cl)
+        self.update_build_info(device)
 
     def _set_project_changelist_is_synched(self, is_synched: bool):
         self._is_project_synched = is_synched
@@ -3690,26 +3764,40 @@ class DeviceWidgetUnreal(DeviceWidget):
         needs_resync = not self._is_project_synched or not self._is_engine_synched
         sb_widgets.set_qt_property(self.sync_button, 'not_synched', needs_resync)
 
-    def update_build_info(self, synched_cl: str, built_cl: str):
-        if built_cl is not None and synched_cl is not None and CONFIG.ENGINE_SYNC_METHOD.get_value() == EngineSyncMethod.Build_Engine.value:
-            try:
-                earlier_cl = min(int(built_cl), int(synched_cl))
-                later_cl = max(int(built_cl), int(synched_cl))
-                self._needs_rebuild = p4_changelist_inspection.has_source_code_changes(
-                    earlier_cl,
-                    later_cl,
-                    CONFIG.P4_ENGINE_PATH.get_value()
-                )
-            except p4_changelist_inspection.P4Error as error:
-                LOGGER.error(f"Couldn't check {built_cl} - {built_cl} for source code changes."
-                             f" Reason: {error.message}")
-                # Assume that non-equal CL numbers have source changes
-                self._needs_rebuild = int(built_cl) != int(synched_cl)
+    def update_build_info(
+        self,
+        device: DeviceUnreal,
+    ):
+        self._needs_rebuild = False
+        details = ''
+
+        if CONFIG.ENGINE_SYNC_METHOD.get_value() == EngineSyncMethod.Build_Engine.value:
+            eng_differs = device._built_eng_changelist != device._synced_changelist
+            eng_code_differs = device._built_eng_compatible_cl != device._synced_compatible_cl
+            if (eng_differs or eng_code_differs):
+                self._needs_rebuild = True
+                details += (f"\nBuilt engine CL {device._built_eng_changelist}"
+                            f" (Code {device._built_eng_compatible_cl})")
+
+            proj_differs = device._built_proj_changelist != device._synced_changelist
+            proj_code_differs = device._built_proj_compatible_cl != device._synced_compatible_cl
+            if (proj_differs or proj_code_differs):
+                self._needs_rebuild = True
+                details += (f"\nBuilt project CL {device._built_proj_changelist}"
+                            f" (Code {device._built_proj_compatible_cl})")
+
+        desired_tooltip = 'Build changelist'
+
+        if self._needs_rebuild:
+            desired_tooltip += "\n\nBuild required.\n\n"
         else:
-            self._needs_rebuild = False
+            desired_tooltip += " (not required - built CL matches synced CL)\n\n"
+
+        desired_tooltip += (f"Synced: {device._synced_changelist} "
+                            f"(Code {device._synced_compatible_cl})")
+
+        desired_tooltip += details
 
-        desired_tooltip = f"Build changelist.\n\nBuild required.\nSynched: {synched_cl}\nBuilt: {built_cl}" \
-            if self._needs_rebuild else f"Build changelist (not required - built CL {built_cl} matches synched CL)"
         self._desired_build_button_tooltip = desired_tooltip
         self._update_build_button_tooltip()
         self._refresh_build_info_ui()
@@ -3746,10 +3834,10 @@ class DeviceWidgetUnreal(DeviceWidget):
         tooltip = base_tooltip
 
         if not is_synched:
-            tooltip += "\nNot synched to selected CL"
+            tooltip += "\nNot synced to selected CL"
 
         if self._needs_rebuild:
-            tooltip += "\nSynched CL not built"
+            tooltip += "\nSynced CL not built"
 
         label.setToolTip(tooltip)
 
diff --git a/Engine/Plugins/VirtualProduction/Switchboard/Source/Switchboard/switchboard/listener_client.py b/Engine/Plugins/VirtualProduction/Switchboard/Source/Switchboard/switchboard/listener_client.py
--- a/Engine/Plugins/VirtualProduction/Switchboard/Source/Switchboard/switchboard/listener_client.py
+++ b/Engine/Plugins/VirtualProduction/Switchboard/Source/Switchboard/switchboard/listener_client.py
@@ -446,7 +446,7 @@ class ListenerClient(QtCore.QObject):
             self.protocol.client = self
 
             (self.reader, self.writer) = await protocol.create_stream()
-            self.reader._limit = 2 ** 20  # 1 MiB max message (default 64 KiB)
+            self.reader._limit = 10 * (2 ** 20)  # 10 MiB max message (default 64 KiB)
             message_recv_task = asyncio.create_task(self._message_recv_loop())
 
             self.is_authenticated = await self.protocol.authenticate()
diff --git a/Engine/Plugins/VirtualProduction/Switchboard/Source/Switchboard/switchboard/ue_plugin_utils.py b/Engine/Plugins/VirtualProduction/Switchboard/Source/Switchboard/switchboard/ue_plugin_utils.py
--- a/Engine/Plugins/VirtualProduction/Switchboard/Source/Switchboard/switchboard/ue_plugin_utils.py
+++ b/Engine/Plugins/VirtualProduction/Switchboard/Source/Switchboard/switchboard/ue_plugin_utils.py
@@ -6,12 +6,14 @@ from enum import auto, Enum
 import glob
 import json
 import logging
-from pathlib import Path, PurePosixPath
+from pathlib import Path, PurePath, PurePosixPath
 import os
 import time
-from typing import Any, Optional, Union
+from typing import Any, Optional, TypeAlias, Union
 
 
+UnrealPath: TypeAlias = PurePosixPath
+
 PROJECT_CONTENT_MOUNT_NAME = 'Game'
 PROJECT_CONTENT_MOUNT_POINT = f'/{PROJECT_CONTENT_MOUNT_NAME}'
 
@@ -48,7 +50,7 @@ class UnrealPluginManager:
         self._uproject_path: Optional[Path] = None
 
         # Information derived from .uproject
-        self._addl_plugin_dirs: list[Path] = []
+        self._addl_plugin_dirs: list[PurePath] = []
         self._disabled_plugins: set[str] = set()
 
         # Plugin name -> UnrealPlugin
@@ -109,11 +111,13 @@ class UnrealPluginManager:
         self._project_plugin_paths = self.find_plugins_in_dir(proj_plugin_path)
 
         # Parse .uproject additional plugin paths
-        addl_dirs: list[str] = uproj.get('AdditionalPluginDirectories', [])
-        self._addl_plugin_dirs = [proj_dir / x for x in addl_dirs]
+        self._addl_plugin_dirs = [
+            PurePath(x) for x in
+            uproj.get('AdditionalPluginDirectories', [])
+        ]
 
         for addl_dir in self._addl_plugin_dirs:
-            addl_plugins = self.find_plugins_in_dir(addl_dir)
+            addl_plugins = self.find_plugins_in_dir(proj_dir / addl_dir)
             self._external_plugin_paths.extend(addl_plugins)
 
         # Parse .uproject plugin references to identify disabled plugins
@@ -155,7 +159,7 @@ class UnrealPluginManager:
 
     def content_to_file_path(
         self,
-        path: Union[PurePosixPath, str]
+        path: Union[UnrealPath, str]
     ) -> Optional[Path]:
         '''
         Given an Unreal content path beginning with e.g. `/Game/...` or
@@ -163,7 +167,7 @@ class UnrealPluginManager:
         '''
 
         if isinstance(path, str):
-            path = PurePosixPath(path)
+            path = UnrealPath(path)
 
         if not path.is_absolute() or len(path.parts) < 2:
             logging.error('Not an absolute path: %s', path)
@@ -191,7 +195,7 @@ class UnrealPluginManager:
     def file_to_content_path(
         self,
         path: Union[Path, str]
-    ) -> Optional[PurePosixPath]:
+    ) -> Optional[UnrealPath]:
         '''
         Given a file or directory path, map it to the corresponding Unreal
         content path (relative to the project or plugin content mount point).
@@ -205,7 +209,7 @@ class UnrealPluginManager:
 
         if path.is_relative_to(project_content_path):
             rest = path.relative_to(project_content_path)
-            return PurePosixPath(PROJECT_CONTENT_MOUNT_POINT) / rest
+            return UnrealPath(PROJECT_CONTENT_MOUNT_POINT) / rest
 
         for parent in path.parents:
             if plugin := self.get_plugin_by_content_dir(parent):
@@ -249,6 +253,11 @@ class UnrealPluginManager:
     def enabled_project_plugins(self) -> list[UnrealPlugin]:
         return self._enabled_project_plugins
 
+    @property
+    def addl_plugin_dirs(self) -> list[PurePath]:
+        ''' Returns the relative path fragments from the uproject JSON. '''
+        return self._addl_plugin_dirs
+
     def _refresh_plugin_map(self):
         self._plugin_map.clear()
         self._content_dir_map.clear()
@@ -431,7 +440,7 @@ class UnrealPlugin(object):
         return self.plugin_dir / 'Content'
 
     @property
-    def mounted_path(self) -> PurePosixPath:
+    def mounted_path(self) -> UnrealPath:
         '''
         The root content path of the plugin when its "Content" directory is
         mounted in Unreal Engine.
@@ -439,7 +448,7 @@ class UnrealPlugin(object):
         # Note that UE uses the file name of the uplugin file to produce the
         # mounted path in engine, and not the name of the plugin directory or
         # any data inside the uplugin file.
-        return PurePosixPath(f'/{self.name}')
+        return UnrealPath(f'/{self.name}')
 
     @property
     def plugin_type(self) -> UnrealPluginType:
diff --git a/Engine/Plugins/VirtualProduction/Switchboard/Source/Switchboard/switchboard/util/p4_changelist_inspection.py b/Engine/Plugins/VirtualProduction/Switchboard/Source/Switchboard/switchboard/util/p4_changelist_inspection.py
deleted file mode 100644
