Switchboard: Fix for missed file from CL 40217760.

diff --git a/Engine/Plugins/VirtualProduction/Switchboard/Source/Switchboard/switchboard/switchboard_dialog.py b/Engine/Plugins/VirtualProduction/Switchboard/Source/Switchboard/switchboard/switchboard_dialog.py
--- a/Engine/Plugins/VirtualProduction/Switchboard/Source/Switchboard/switchboard/switchboard_dialog.py
+++ b/Engine/Plugins/VirtualProduction/Switchboard/Source/Switchboard/switchboard/switchboard_dialog.py
@@ -1454,7 +1454,7 @@ class SwitchboardDialog(QtCore.QObject):
         device.device_qt_handler.signal_device_client_disconnected.connect(self.device_client_disconnected, QtCore.Qt.QueuedConnection)
         device.device_qt_handler.signal_device_project_changelist_changed.connect(self.device_project_changelist_changed, QtCore.Qt.QueuedConnection)
         device.device_qt_handler.signal_device_engine_changelist_changed.connect(self.device_engine_changelist_changed, QtCore.Qt.QueuedConnection)
-        device.device_qt_handler.signal_device_built_engine_changelist_changed.connect(self.device_built_engine_changelist_changed, QtCore.Qt.QueuedConnection)
+        device.device_qt_handler.signal_device_built_changelist_changed.connect(self.device_built_engine_changelist_changed, QtCore.Qt.QueuedConnection)
         device.device_qt_handler.signal_device_status_changed.connect(self.device_status_changed, QtCore.Qt.QueuedConnection)
         device.device_qt_handler.signal_device_sync_failed.connect(self.device_sync_failed, QtCore.Qt.QueuedConnection)
         device.device_qt_handler.signal_device_is_recording_device_changed.connect(self.device_is_recording_device_changed, QtCore.Qt.QueuedConnection)
