[Backout] - CL40217760
#rnx
Original CL Desc
-----------------------------------------------------------------
Switchboard: Build indicator fixes, code cleanup, and documentation.

Previously, the `CompatibleChangelist` value from `UnrealEditor.version` was regarded as the source of truth for whether the build was up-to-date, which was inadequate. At the same time, the code to determine the expected value was was more complicated than it needed to be.

Now, we compare both `UnrealEditor.version` and `$(ProjectName)Editor.target` against `Build.version`, and for each file, we compare both the `Changelist` and `CompatibleChangelist` values. Switchboard does not need the code changelist determination logic, because that determination is already made at sync time by `sbl_helper.py`, and the computed value cached in `CompatibleChangelist`.

Also refactors `UnrealPluginManager` to expose `AdditionalPluginDirectories` from the .uproject (which will likely be required for some sync indicator fixes in the near future), and fixes some linter warnings where the built-in `id()` function was being shadowed.

diff --git a/Engine/Plugins/VirtualProduction/Switchboard/Source/Switchboard/docs/SyncAndVersioning.md b/Engine/Plugins/VirtualProduction/Switchboard/Source/Switchboard/docs/SyncAndVersioning.md
deleted file mode 100644
diff --git a/Engine/Plugins/VirtualProduction/Switchboard/Source/Switchboard/switchboard/devices/device_base.py b/Engine/Plugins/VirtualProduction/Switchboard/Source/Switchboard/switchboard/devices/device_base.py
--- a/Engine/Plugins/VirtualProduction/Switchboard/Source/Switchboard/switchboard/devices/device_base.py
+++ b/Engine/Plugins/VirtualProduction/Switchboard/Source/Switchboard/switchboard/devices/device_base.py
@@ -55,7 +55,7 @@ class DeviceQtHandler(QtCore.QObject):
     signal_device_client_disconnected = QtCore.Signal(object)
     signal_device_project_changelist_changed = QtCore.Signal(object)
     signal_device_engine_changelist_changed = QtCore.Signal(object)
-    signal_device_built_changelist_changed = QtCore.Signal(object)
+    signal_device_built_engine_changelist_changed = QtCore.Signal(object)
     signal_device_sync_failed = QtCore.Signal(object)
     signal_device_is_recording_device_changed = QtCore.Signal(object, int)
 
@@ -105,6 +105,7 @@ class Device(QtCore.QObject):
 
         self._project_changelist: Optional[int] = None
         self._engine_changelist: Optional[int] = None
+        self._built_engine_changelist: Optional[int] = None
 
         self._status = DeviceStatus.DISCONNECTED
 
@@ -280,6 +281,17 @@ class Device(QtCore.QObject):
         self.device_qt_handler.signal_device_engine_changelist_changed.emit(
             self)
 
+    @property
+    def built_engine_changelist(self):
+        return self._built_engine_changelist
+
+    @built_engine_changelist.setter
+    def built_engine_changelist(self, value: Optional[int]):
+        self._built_engine_changelist = value
+        self.device_qt_handler.signal_device_built_engine_changelist_changed.emit(
+            self
+        )
+
     @property
     def device_osc_port(self) -> int:
         raise NotImplementedError()
diff --git a/Engine/Plugins/VirtualProduction/Switchboard/Source/Switchboard/switchboard/devices/unreal/plugin_unreal.py b/Engine/Plugins/VirtualProduction/Switchboard/Source/Switchboard/switchboard/devices/unreal/plugin_unreal.py
--- a/Engine/Plugins/VirtualProduction/Switchboard/Source/Switchboard/switchboard/devices/unreal/plugin_unreal.py
+++ b/Engine/Plugins/VirtualProduction/Switchboard/Source/Switchboard/switchboard/devices/unreal/plugin_unreal.py
@@ -43,6 +43,7 @@ from switchboard.devices.unreal.uassetparser import UassetParser
 from switchboard.listener_client import ListenerClient
 from switchboard.switchboard_logging import LOGGER
 from switchboard.tools.insights_launcher import InsightsLauncher
+from switchboard.util import p4_changelist_inspection
 from . import version_helpers
 from .listener_watcher import ListenerWatcher
 from .redeploy_dialog import RedeployListenerDialog
@@ -486,11 +487,11 @@ class MediaProfileSetting(LiveLinkPresetSetting):
 
 class SyncCategoryOption:
     ''' Used to represent a UGS sync filter category within a Setting. '''
-    def __init__(self, catid: Union[str, uuid.UUID], name: str):
-        if isinstance(catid, str):
-            catid = uuid.UUID(catid)
+    def __init__(self, id: Union[str, uuid.UUID], name: str):
+        if isinstance(id, str):
+            id = uuid.UUID(id)
 
-        self.catid = catid
+        self.id = id
         self.name = name
 
     def __str__(self) -> str:
@@ -500,7 +501,7 @@ class SyncCategoryOption:
     def __eq__(self, other: object):
         ''' Two category options compare equal if their UUIDs are equal. '''
         if isinstance(other, SyncCategoryOption):
-            return self.catid == other.catid
+            return self.id == other.id
         else:
             return False
 
@@ -760,7 +761,7 @@ class DeviceUnreal(Device):
             tool_tip='UnrealGameSync filter categories to include during sync',
             category="Source Control Settings",
         ).with_get_json_override_fn(
-            lambda val_list: [(str(x.catid), x.name) for x in val_list]
+            lambda val_list: [(str(x.id), x.name) for x in val_list]
         ).with_config_set_override_fn(
             lambda val_list: [
                 SyncCategoryOption(uuid.UUID(x[0]), x[1])
@@ -785,16 +786,6 @@ class DeviceUnreal(Device):
                 'this Switchboard configuration \n'),
             category="UE Settings",
         ),
-        'editor_target_override': StringSetting(
-            attr_name='editor_target_override',
-            nice_name='Editor Target Name Override',
-            value='',
-            placeholder_text='$(ProjectName)Editor',
-            tool_tip=(
-                "The name of your project's editor target. Typically $(ProjectName)Editor (which is inferred "
-                "automatically). This override only needs to be set if deviating from the standard convention."),
-            category='UE Settings',
-        )
     }
 
     if ENABLE_UGS_SUPPORT:
@@ -1082,33 +1073,10 @@ class DeviceUnreal(Device):
         self.auto_connect = False
 
         self.runtime_str = ""
-
-        # The CLs returned via p4 cstat query, or optimistically
-        # updated with the "inflight" values upon successful sync
-        # (Defined in the base class, documented here for context)
-        #self._project_changelist: Optional[int] = None
-        #self._engine_changelist: Optional[int] = None
-
-        # The CLs for which outstanding p4 syncs have been queued
         self.inflight_project_cl = None
         self.inflight_engine_cl = None
         self.inflight_sync_hash: Optional[str] = None
 
-        # The fields read from Build.version JSON
-        # (Written during sync by UGS/sbl_helper.py)
-        self._synced_changelist: Optional[int] = None
-        self._synced_compatible_cl: Optional[int] = None
-
-        # The fields read from UnrealEditor.version JSON
-        # (Written by UBT following a successful build of ANY project)
-        self._built_eng_changelist: Optional[int] = None
-        self._built_eng_compatible_cl: Optional[int] = None
-
-        # The fields read from (Project)Editor.target JSON
-        # (Written by UBT following a successful build of the active project)
-        self._built_proj_changelist: Optional[int] = None
-        self._built_proj_compatible_cl: Optional[int] = None
-
         listener_qt_handler = self.unreal_client.listener_qt_handler
         listener_qt_handler.listener_connecting.connect(
             super().connecting_listener, QtCore.Qt.QueuedConnection)
@@ -1492,7 +1460,7 @@ class DeviceUnreal(Device):
                 self._request_engine_changelist_number()
         else:
             self.widget.engine_changelist_label.hide()
-            self.widget.update_build_info(self)
+            self.widget.update_build_info(synched_cl=self.engine_changelist, built_cl=self.built_engine_changelist)
 
     def on_setting_exclude_from_build_changed(self, exclude_from_build):
         self.widget.update_exclude_from_build(exclude_from_build, not self.is_disconnected)
@@ -1532,36 +1500,17 @@ class DeviceUnreal(Device):
         _, msg = message_protocol.create_copy_file_from_listener_message(roles_file_path)
         self.unreal_client.send_message(msg)
 
-    def _request_version_files(self):
+    def _request_unreal_editor_version_file(self):
         '''
-        Requests the Build.version, UnrealEditor.version and $(Proj)Editor.target files.
-        On success, _on_receive_version_file is called with the file content.
+        Requests the Engine/Binaries/[platform]/UnrealEditor.version file.
+        On success _on_receive_editor_version is called with the file content.
         '''
-        engine_dir = CONFIG.ENGINE_DIR.get_value(self.name)
-        uproj_path = CONFIG.UPROJECT_PATH.get_value(self.name)
-        proj_dir = os.path.dirname(uproj_path)
+        engine_path = CONFIG.ENGINE_DIR.get_value(self.name).replace('"', '')
         platform_dir = self.platform_binary_directory
-
-        proj_name, _ = os.path.splitext(os.path.basename(uproj_path))
-        editor_target = DeviceUnreal.csettings['editor_target_override'].get_value()
-        if not editor_target:
-            editor_target = f'{proj_name}Editor'
-
-        build_version_path = os.path.join(engine_dir, 'Build', 'Build.version')
-        _, msg = message_protocol.create_copy_file_from_listener_message(build_version_path)
-        self.unreal_client.send_message(msg)
-
-        editor_version_path = os.path.join(engine_dir, 'Binaries', platform_dir,
-                                           'UnrealEditor.version')
-        _, msg = message_protocol.create_copy_file_from_listener_message(editor_version_path)
-        self.unreal_client.send_message(msg)
-
-        proj_target_path = os.path.join(proj_dir, 'Binaries', platform_dir,
-                                        f'{editor_target}.target')
-        _, msg = message_protocol.create_copy_file_from_listener_message(proj_target_path)
+        roles_file_path = os.path.join(os.path.dirname(engine_path), "Engine", "Binaries", platform_dir, "UnrealEditor.version")
+        _, msg = message_protocol.create_copy_file_from_listener_message(roles_file_path)
         self.unreal_client.send_message(msg)
 
-
     def _request_project_changelist_number(self):
         if not CONFIG.P4_ENABLED.get_value():
             return
@@ -1798,8 +1747,8 @@ class DeviceUnreal(Device):
             incl_categories: list[SyncCategoryOption] = self.csettings[
                 'included_sync_categories'].get_value()
             if incl_categories:
-                incl_categories.sort(key=lambda cat: cat.catid)
-                include_str = ','.join(str(cat.catid) for cat in incl_categories)
+                incl_categories.sort(key=lambda cat: cat.id)
+                include_str = ','.join(str(cat.id) for cat in incl_categories)
                 sync_args += f' --include-categories="{include_str}"'
 
             if custom_view_str:
@@ -2638,7 +2587,7 @@ class DeviceUnreal(Device):
 
             # If you build and sync the engine, update its CL
             if CONFIG.ENGINE_SYNC_METHOD.get_value() == EngineSyncMethod.Build_Engine.value:
-                self._request_version_files()
+
                 if self.inflight_engine_cl:
                     self.engine_changelist = self.inflight_engine_cl
                     self.inflight_engine_cl = None
@@ -2682,7 +2631,7 @@ class DeviceUnreal(Device):
                 self._request_project_changelist_number()
                 if CONFIG.ENGINE_SYNC_METHOD.get_value() == EngineSyncMethod.Build_Engine.value:
                     self._request_engine_changelist_number()
-                    self._request_version_files()
+                    self._request_unreal_editor_version_file()
 
         elif UnrealJobs.CstatPrefix.value in program_name:
             output = get_stdout_str()
@@ -2748,9 +2697,9 @@ class DeviceUnreal(Device):
     def on_file_received(self, source_path, content):
         if source_path.endswith(DeviceUnreal.csettings["roles_filename"].get_value(self.name)):
             self._on_receive_roles(content)
-        elif source_path.endswith('.version') or source_path.endswith('.target'):
-            self._on_receive_version_file(source_path, content)
-
+        elif "UnrealEditor" in source_path and source_path.endswith(".version"):
+            self._on_receive_editor_version(content)
+        
     def _on_receive_roles(self, content):
         decoded_content = base64.b64decode(content).decode()
         tags = parse_unreal_tag_file(decoded_content.splitlines())
@@ -2767,31 +2716,15 @@ class DeviceUnreal(Device):
                 f"{self.name}: Please change the roles for this device in "
                 "the settings or in the unreal project settings!")
 
-    def _on_receive_version_file(self, source_path: str, content):
+    def _on_receive_editor_version(self, content):
         '''
-        Handles .version/.target files (see _request_version_files)
+        Receives the Engine/Binaries/[platform]/UnrealEditor.version file _request_unreal_editor_version_file
         '''
         decoded_content = base64.b64decode(content).decode()
         data = json.loads(decoded_content)
-
-        if source_path.endswith('.target'):
-            data = data.get('Version')
-
-        changelist = data.get('Changelist')
-        compatible_cl = data.get('CompatibleChangelist')
-
+        compatible_cl = data.get("CompatibleChangelist", None)
         for device in self.devices_sharing_workspace():
-            if source_path.endswith('Build.version'):
-                device._synced_changelist = changelist
-                device._synced_compatible_cl = compatible_cl
-            elif source_path.endswith('UnrealEditor.version'):
-                device._built_eng_changelist = changelist
-                device._built_eng_compatible_cl = compatible_cl
-            elif source_path.endswith('.target'):
-                device._built_proj_changelist = changelist
-                device._built_proj_compatible_cl = compatible_cl
-
-        self.device_qt_handler.signal_device_built_changelist_changed.emit(self)
+            device.built_engine_changelist = compatible_cl
 
     def on_file_receive_failed(self, source_path, error):
         if source_path.endswith(DeviceUnreal.csettings["roles_filename"].get_value(self.name)):
@@ -2800,8 +2733,7 @@ class DeviceUnreal(Device):
                 LOGGER.error(
                     f"{self.name}: Error receiving role file from listener and "
                     f"device claims to have these roles: {' | '.join(roles)}")
-
-        LOGGER.error(f"{self.name}: Error retrieving file {source_path}: {error}")
+                LOGGER.error(f"Error: {error}")
 
     def on_listener_programstdout(self, message):
         ''' Handles updates to stdout of programs
@@ -2978,7 +2910,7 @@ class DeviceUnreal(Device):
 
         if CONFIG.ENGINE_SYNC_METHOD.get_value() == EngineSyncMethod.Build_Engine.value:
             self._request_engine_changelist_number()
-            self._request_version_files()
+            self._request_unreal_editor_version_file()
 
     def transport_paths(self, device_recording):
         '''
@@ -3725,26 +3657,20 @@ class DeviceWidgetUnreal(DeviceWidget):
 
         self.project_changelist_label.show()
 
-    def update_engine_changelist(
-        self,
-        required_cl: str,
-        device: DeviceUnreal,
-    ):
+    def update_engine_changelist(self, required_cl: str, synched_cl: str, built__cl: str):
         if not CONFIG.ENGINE_SYNC_METHOD.get_value() == EngineSyncMethod.Build_Engine.value:
             return
 
-        synced_cl = device.engine_changelist
-
-        self.engine_changelist_label.setText(f'E: {synced_cl}')
+        self.engine_changelist_label.setText(f'E: {synched_cl}')
         if not self.exclude_from_build:
             self.engine_changelist_label.show()
 
         self.sync_button.show()
         self.build_button.show()
 
-        is_synched = required_cl is None or required_cl == synced_cl
+        is_synched = required_cl is None or required_cl == synched_cl
         self._set_engine_changelist_is_synched(is_synched)
-        self.update_build_info(device)
+        self.update_build_info(synched_cl=synched_cl, built_cl=built__cl)
 
     def _set_project_changelist_is_synched(self, is_synched: bool):
         self._is_project_synched = is_synched
@@ -3764,40 +3690,26 @@ class DeviceWidgetUnreal(DeviceWidget):
         needs_resync = not self._is_project_synched or not self._is_engine_synched
         sb_widgets.set_qt_property(self.sync_button, 'not_synched', needs_resync)
 
-    def update_build_info(
-        self,
-        device: DeviceUnreal,
-    ):
-        self._needs_rebuild = False
-        details = ''
-
-        if CONFIG.ENGINE_SYNC_METHOD.get_value() == EngineSyncMethod.Build_Engine.value:
-            eng_differs = device._built_eng_changelist != device._synced_changelist
-            eng_code_differs = device._built_eng_compatible_cl != device._synced_compatible_cl
-            if (eng_differs or eng_code_differs):
-                self._needs_rebuild = True
-                details += (f"\nBuilt engine CL {device._built_eng_changelist}"
-                            f" (Code {device._built_eng_compatible_cl})")
-
-            proj_differs = device._built_proj_changelist != device._synced_changelist
-            proj_code_differs = device._built_proj_compatible_cl != device._synced_compatible_cl
-            if (proj_differs or proj_code_differs):
-                self._needs_rebuild = True
-                details += (f"\nBuilt project CL {device._built_proj_changelist}"
-                            f" (Code {device._built_proj_compatible_cl})")
-
-        desired_tooltip = 'Build changelist'
-
-        if self._needs_rebuild:
-            desired_tooltip += "\n\nBuild required.\n\n"
+    def update_build_info(self, synched_cl: str, built_cl: str):
+        if built_cl is not None and synched_cl is not None and CONFIG.ENGINE_SYNC_METHOD.get_value() == EngineSyncMethod.Build_Engine.value:
+            try:
+                earlier_cl = min(int(built_cl), int(synched_cl))
+                later_cl = max(int(built_cl), int(synched_cl))
+                self._needs_rebuild = p4_changelist_inspection.has_source_code_changes(
+                    earlier_cl,
+                    later_cl,
+                    CONFIG.P4_ENGINE_PATH.get_value()
+                )
+            except p4_changelist_inspection.P4Error as error:
+                LOGGER.error(f"Couldn't check {built_cl} - {built_cl} for source code changes."
+                             f" Reason: {error.message}")
+                # Assume that non-equal CL numbers have source changes
+                self._needs_rebuild = int(built_cl) != int(synched_cl)
         else:
-            desired_tooltip += " (not required - built CL matches synced CL)\n\n"
-
-        desired_tooltip += (f"Synced: {device._synced_changelist} "
-                            f"(Code {device._synced_compatible_cl})")
-
-        desired_tooltip += details
+            self._needs_rebuild = False
 
+        desired_tooltip = f"Build changelist.\n\nBuild required.\nSynched: {synched_cl}\nBuilt: {built_cl}" \
+            if self._needs_rebuild else f"Build changelist (not required - built CL {built_cl} matches synched CL)"
         self._desired_build_button_tooltip = desired_tooltip
         self._update_build_button_tooltip()
         self._refresh_build_info_ui()
@@ -3834,10 +3746,10 @@ class DeviceWidgetUnreal(DeviceWidget):
         tooltip = base_tooltip
 
         if not is_synched:
-            tooltip += "\nNot synced to selected CL"
+            tooltip += "\nNot synched to selected CL"
 
         if self._needs_rebuild:
-            tooltip += "\nSynced CL not built"
+            tooltip += "\nSynched CL not built"
 
         label.setToolTip(tooltip)
 
diff --git a/Engine/Plugins/VirtualProduction/Switchboard/Source/Switchboard/switchboard/listener_client.py b/Engine/Plugins/VirtualProduction/Switchboard/Source/Switchboard/switchboard/listener_client.py
--- a/Engine/Plugins/VirtualProduction/Switchboard/Source/Switchboard/switchboard/listener_client.py
+++ b/Engine/Plugins/VirtualProduction/Switchboard/Source/Switchboard/switchboard/listener_client.py
@@ -446,7 +446,7 @@ class ListenerClient(QtCore.QObject):
             self.protocol.client = self
 
             (self.reader, self.writer) = await protocol.create_stream()
-            self.reader._limit = 10 * (2 ** 20)  # 10 MiB max message (default 64 KiB)
+            self.reader._limit = 2 ** 20  # 1 MiB max message (default 64 KiB)
             message_recv_task = asyncio.create_task(self._message_recv_loop())
 
             self.is_authenticated = await self.protocol.authenticate()
diff --git a/Engine/Plugins/VirtualProduction/Switchboard/Source/Switchboard/switchboard/ue_plugin_utils.py b/Engine/Plugins/VirtualProduction/Switchboard/Source/Switchboard/switchboard/ue_plugin_utils.py
--- a/Engine/Plugins/VirtualProduction/Switchboard/Source/Switchboard/switchboard/ue_plugin_utils.py
+++ b/Engine/Plugins/VirtualProduction/Switchboard/Source/Switchboard/switchboard/ue_plugin_utils.py
@@ -6,14 +6,12 @@ from enum import auto, Enum
 import glob
 import json
 import logging
-from pathlib import Path, PurePath, PurePosixPath
+from pathlib import Path, PurePosixPath
 import os
 import time
-from typing import Any, Optional, TypeAlias, Union
+from typing import Any, Optional, Union
 
 
-UnrealPath: TypeAlias = PurePosixPath
-
 PROJECT_CONTENT_MOUNT_NAME = 'Game'
 PROJECT_CONTENT_MOUNT_POINT = f'/{PROJECT_CONTENT_MOUNT_NAME}'
 
@@ -50,7 +48,7 @@ class UnrealPluginManager:
         self._uproject_path: Optional[Path] = None
 
         # Information derived from .uproject
-        self._addl_plugin_dirs: list[PurePath] = []
+        self._addl_plugin_dirs: list[Path] = []
         self._disabled_plugins: set[str] = set()
 
         # Plugin name -> UnrealPlugin
@@ -111,13 +109,11 @@ class UnrealPluginManager:
         self._project_plugin_paths = self.find_plugins_in_dir(proj_plugin_path)
 
         # Parse .uproject additional plugin paths
-        self._addl_plugin_dirs = [
-            PurePath(x) for x in
-            uproj.get('AdditionalPluginDirectories', [])
-        ]
+        addl_dirs: list[str] = uproj.get('AdditionalPluginDirectories', [])
+        self._addl_plugin_dirs = [proj_dir / x for x in addl_dirs]
 
         for addl_dir in self._addl_plugin_dirs:
-            addl_plugins = self.find_plugins_in_dir(proj_dir / addl_dir)
+            addl_plugins = self.find_plugins_in_dir(addl_dir)
             self._external_plugin_paths.extend(addl_plugins)
 
         # Parse .uproject plugin references to identify disabled plugins
@@ -159,7 +155,7 @@ class UnrealPluginManager:
 
     def content_to_file_path(
         self,
-        path: Union[UnrealPath, str]
+        path: Union[PurePosixPath, str]
     ) -> Optional[Path]:
         '''
         Given an Unreal content path beginning with e.g. `/Game/...` or
@@ -167,7 +163,7 @@ class UnrealPluginManager:
         '''
 
         if isinstance(path, str):
-            path = UnrealPath(path)
+            path = PurePosixPath(path)
 
         if not path.is_absolute() or len(path.parts) < 2:
             logging.error('Not an absolute path: %s', path)
@@ -195,7 +191,7 @@ class UnrealPluginManager:
     def file_to_content_path(
         self,
         path: Union[Path, str]
-    ) -> Optional[UnrealPath]:
+    ) -> Optional[PurePosixPath]:
         '''
         Given a file or directory path, map it to the corresponding Unreal
         content path (relative to the project or plugin content mount point).
@@ -209,7 +205,7 @@ class UnrealPluginManager:
 
         if path.is_relative_to(project_content_path):
             rest = path.relative_to(project_content_path)
-            return UnrealPath(PROJECT_CONTENT_MOUNT_POINT) / rest
+            return PurePosixPath(PROJECT_CONTENT_MOUNT_POINT) / rest
 
         for parent in path.parents:
             if plugin := self.get_plugin_by_content_dir(parent):
@@ -253,11 +249,6 @@ class UnrealPluginManager:
     def enabled_project_plugins(self) -> list[UnrealPlugin]:
         return self._enabled_project_plugins
 
-    @property
-    def addl_plugin_dirs(self) -> list[PurePath]:
-        ''' Returns the relative path fragments from the uproject JSON. '''
-        return self._addl_plugin_dirs
-
     def _refresh_plugin_map(self):
         self._plugin_map.clear()
         self._content_dir_map.clear()
@@ -440,7 +431,7 @@ class UnrealPlugin(object):
         return self.plugin_dir / 'Content'
 
     @property
-    def mounted_path(self) -> UnrealPath:
+    def mounted_path(self) -> PurePosixPath:
         '''
         The root content path of the plugin when its "Content" directory is
         mounted in Unreal Engine.
@@ -448,7 +439,7 @@ class UnrealPlugin(object):
         # Note that UE uses the file name of the uplugin file to produce the
         # mounted path in engine, and not the name of the plugin directory or
         # any data inside the uplugin file.
-        return UnrealPath(f'/{self.name}')
+        return PurePosixPath(f'/{self.name}')
 
     @property
     def plugin_type(self) -> UnrealPluginType:
diff --git a/Engine/Plugins/VirtualProduction/Switchboard/Source/Switchboard/switchboard/util/p4_changelist_inspection.py b/Engine/Plugins/VirtualProduction/Switchboard/Source/Switchboard/switchboard/util/p4_changelist_inspection.py
new file mode 100644
--- /dev/null
+++ b/Engine/Plugins/VirtualProduction/Switchboard/Source/Switchboard/switchboard/util/p4_changelist_inspection.py
@@ -0,0 +1,105 @@
+# Copyright Epic Games, Inc. All Rights Reserved.
+import os
+import subprocess
+import typing
+from enum import IntFlag
+
+from ..config import CONFIG
+from ..p4_utils import (run, valueForMarshalledKey, hasValueForMarshalledKey,
+                        ALL_CODE_EXTS)
+from ..switchboard_utils import get_hidden_sp_startupinfo
+
+
+def has_source_code_changes(earlier_cl: int, later_cl: int, engine_path: str) -> bool:
+    """
+    Checks whether the CLs between current_cl (exclusive) and future_cl (inclusive) are content only.
+    
+    Raises P4Error if the data cannot be obtained.
+    """
+    if earlier_cl == later_cl:
+        return False
+    
+    return any_of_changelists(
+        earlier_cl,
+        later_cl,
+        engine_path,
+        lambda cl_description:
+            EChangelistFileFlags.Code in determine_changelist_type(valueForMarshalledKey(cl_description, "change")),
+        [earlier_cl]
+    )
+
+def any_of_changelists(
+    earlier_cl: int,
+    later_cl: int,
+    engine_path: str,
+    predicate_callback: typing.Callable[[typing.Dict], bool],
+    skipped_cls: typing.List[int] = None
+):
+    """
+    Checks whether a predicate holds for all CLs from current_cl to future_cl (both inclusive),
+    The callback can accept the result of an entry returned by "p4 -z tag -G changes".
+    """
+    # Maybe we need to use -m 10 here to only get 10 CLs at a time... if there are a lot of changes p4 may hang up
+    changes: dict = run(f"changes {engine_path}/...@{earlier_cl},{later_cl}")
+    _raise_exception_on_error(changes)
+    
+    for change in changes:
+        cl_number = int(valueForMarshalledKey(change, "change"))
+        should_skip = skipped_cls is not None and cl_number in skipped_cls
+        if not should_skip and predicate_callback(change):
+            return True
+        
+    return False
+
+
+class EChangelistFileFlags(IntFlag):
+    # CL contains source changes
+    Code = 1,
+    # CL contains content changes
+    Content = 2
+    
+    
+def determine_changelist_type(changelist: int) -> EChangelistFileFlags:
+    changes: dict = run(f"describe {changelist}")
+    _raise_exception_on_error(changes)
+
+    has_source_changes = False
+    has_content_changes = False
+    for change in changes:
+        file_index: int = 0
+        # Each changed file path name will have an entry depotFilex
+        while hasValueForMarshalledKey(change, f"depotFile{file_index}"):
+            file_path_name: str = valueForMarshalledKey(change, f"depotFile{file_index}")
+
+            has_source_changes = any(
+                [file_path_name.endswith(ext) for ext in ALL_CODE_EXTS])
+
+            # If it's not a source file, treat it as a content file
+            has_content_changes = not has_source_changes
+
+            # All flags fulfilled - no need to inspect more files
+            if has_source_changes and has_content_changes:
+                break
+
+            file_index += 1
+
+    result = 0
+    if has_source_changes:
+        result |= EChangelistFileFlags.Code
+    if has_content_changes:
+        result |= EChangelistFileFlags.Content
+    return result
+
+
+class P4Error(RuntimeError):
+    def __init__(self, message: str):
+        self.message = message
+
+
+def _raise_exception_on_error(p4_command_result: list):
+    """
+    Accepts the result of p4_util.run and raises an error if any error is set
+    """
+    if len(p4_command_result) > 0 and valueForMarshalledKey(p4_command_result[0], "code") == "error":
+        error_message = valueForMarshalledKey(p4_command_result[0], "data")
+        raise P4Error(error_message)
\ No newline at end of file
