Switchboard: Update uassetparser.py.

Accommodates LegacyFileVersion -9, EUnrealEngineObjectUE5Version::VERSE_CELLS, and EUnrealEngineObjectUE5Version::PACKAGE_SAVED_HASH.

#rnx

diff --git a/Engine/Plugins/VirtualProduction/Switchboard/Source/Switchboard/switchboard/devices/unreal/uassetparser.py b/Engine/Plugins/VirtualProduction/Switchboard/Source/Switchboard/switchboard/devices/unreal/uassetparser.py
--- a/Engine/Plugins/VirtualProduction/Switchboard/Source/Switchboard/switchboard/devices/unreal/uassetparser.py
+++ b/Engine/Plugins/VirtualProduction/Switchboard/Source/Switchboard/switchboard/devices/unreal/uassetparser.py
@@ -9,10 +9,23 @@ Find the relevant C++ code in the UE codebase:
   * PackageUtilities.cpp
 '''
 
+from enum import IntEnum
+import logging
 import os
 import struct
 from typing import BinaryIO
 
+
+LOGGER = logging.getLogger(__name__)
+
+
+class EUnrealEngineObjectUE5Version(IntEnum):
+    ADD_SOFTOBJECTPATH_LIST = 1008
+    METADATA_SERIALIZATION_OFFSET = 1014
+    VERSE_CELLS = 1015
+    PACKAGE_SAVED_HASH = 1016
+
+
 class UassetSummary(object):
     ''' Class to hold the PacketSummary of the asset.
     '''
@@ -20,6 +33,7 @@ class UassetSummary(object):
         from pprint import pprint
         pprint(vars(self))
 
+
 class AssetRegistryData(object):
     ''' Class to hold the asset registry data of the asset
     '''
@@ -28,9 +42,11 @@ class AssetRegistryData(object):
         self.ObjectClassName = ''
         self.tags = {}
 
+
 class FStringError(Exception):
     pass
 
+
 class AssetData(object):
     ''' Class to hold a single asset data from the asset registry
     '''
@@ -38,7 +54,8 @@ class AssetData(object):
         self.AssetClassName = ''
         self.ObjectPathWithoutPackageName = ''
         self.FileOffset = -1
-            
+
+
 class UassetParser(object):
     ''' Class with the logic to parse the asset.
 
@@ -284,7 +301,7 @@ class UassetParser(object):
 
         s.LegacyFileVersion = self.readInt32()
 
-        if s.LegacyFileVersion not in [-7,-8]:
+        if s.LegacyFileVersion not in [-7,-8,-9]:
             raise ValueError('LegacyFileVersion {} is non handled by this parser'.format(s.LegacyFileVersion))
 
         s.LegacyUE3Version = self.readInt32()
@@ -297,12 +314,26 @@ class UassetParser(object):
 
         s.FileVersionLicenseeUE4 = self.readUInt32()
 
-        s.CustomVersions = self.readTArray(lambda : self.f.read(5*4))        
-        s.TotalHeaderSize = self.readInt32()
+        KNOWN_SUPPORTED_UE5VER = EUnrealEngineObjectUE5Version.PACKAGE_SAVED_HASH
+        if s.FileVersionUE5 > KNOWN_SUPPORTED_UE5VER:
+            LOGGER.warning(f"{self.f.name}: ObjectUE5Version {s.FileVersionUE5} too new; "
+                           f"newest known supported version {KNOWN_SUPPORTED_UE5VER}")
+            LOGGER.warning(f"{self.f.name}: Parsing will attempt to continue, but there "
+                           "may be errors reading the file")
+
+        if s.FileVersionUE5 >= EUnrealEngineObjectUE5Version.PACKAGE_SAVED_HASH:
+            s.SavedHash = struct.unpack('<20c', self.f.read(20))
+            s.TotalHeaderSize = self.readInt32()
+
+        s.CustomVersions = self.readTArray(lambda : self.f.read(5*4))
+
+        if s.FileVersionUE5 < EUnrealEngineObjectUE5Version.PACKAGE_SAVED_HASH:
+            s.TotalHeaderSize = self.readInt32()
+
         s.PackageName = self.readFString()
         s.PackageFlags, s.NameCount, s.NameOffset = struct.unpack('<Iii', self.f.read(4*3))
 
-        if s.FileVersionUE5 >= 1008:  # EUnrealEngineObjectUE5Version::ADD_SOFTOBJECTPATH_LIST
+        if s.FileVersionUE5 >= EUnrealEngineObjectUE5Version.ADD_SOFTOBJECTPATH_LIST:
             s.SoftObjectPathsCount = self.readInt32()
             s.SoftObjectPathsOffset = self.readInt32()
 
@@ -312,16 +343,25 @@ class UassetParser(object):
             ''' Convenience local function to calculate remaning bytes in header. Useful for sanity checks during parsing '''
             return s.TotalHeaderSize - self.f.tell() + 1
 
-        s.GatherableTextDataCount, s.GatherableTextDataOffset, s.ExportCount, \
-        s.ExportOffset, s.ImportCount, s.ImportOffset \
+        s.GatherableTextDataCount, s.GatherableTextDataOffset, \
+        s.ExportCount, s.ExportOffset, s.ImportCount, s.ImportOffset \
             = struct.unpack('<6i', self.f.read(6*4))
 
-        if s.FileVersionUE5 >= 1014:  # EUnrealEngineObjectUE5Version::METADATA_SERIALIZATION_OFFSET
+        if s.FileVersionUE5 >= EUnrealEngineObjectUE5Version.VERSE_CELLS:
+            s.CellExportCount, s.CellExportOffset, s.CellImportCount, s.CellImportOffset \
+                = struct.unpack('<4i', self.f.read(4*4))
+
+        if s.FileVersionUE5 >= EUnrealEngineObjectUE5Version.METADATA_SERIALIZATION_OFFSET:
             s.MetaDataOffset = self.readInt32()
 
         s.DependsOffset, s.SoftPackageReferencesCount, s.SoftPackageReferencesOffset, \
-        s.SearchableNamesOffset, s.ThumbnailTableOffset, s.Guid, s.PersistentGuid \
-            = struct.unpack('<5i16s16s', self.f.read(5*4 + 16 + 16))
+        s.SearchableNamesOffset, s.ThumbnailTableOffset \
+            = struct.unpack('<5i', self.f.read(5*4))
+
+        if s.FileVersionUE5 < EUnrealEngineObjectUE5Version.PACKAGE_SAVED_HASH:
+            s.Guid = struct.unpack('<16s', self.f.read(16))
+
+        s.PersistentGuid = struct.unpack('<16s', self.f.read(16))
 
         self.checkFileOffset(s.GatherableTextDataOffset)
         self.checkFileOffset(s.ExportOffset)
