Switchboard: Add VS Code workspace.

This contains inline launch configurations and settings, to aid in debugging and development. The filetype is set to +w in Perforce, so it will be writable in your workspace.

Also marks the switchboard.sh script as filetype +x, which was missed initially.

diff --git a/Engine/Plugins/VirtualProduction/Switchboard/Source/Switchboard/docs/Development.md b/Engine/Plugins/VirtualProduction/Switchboard/Source/Switchboard/docs/Development.md
new file mode 100644
--- /dev/null
+++ b/Engine/Plugins/VirtualProduction/Switchboard/Source/Switchboard/docs/Development.md
@@ -0,0 +1,9 @@
+By opening `Switchboard.code-workspace` in Visual Studio Code, you should find a number of helpful settings preconfigured, including:
+
+  * A debug launch configuration, enabling you to launch Switchboard via F5, set breakpoints, and catch and inspect unhandled exceptions
+
+    * > 🛈 **Note**: The default Python interpreter path is set to the standard Switchboard virtual environment. This ensures any prerequisite packages are available, but assumes you've run `switchboard.bat`/`switchboard.sh` at least once to complete the standard first-time setup prior to debugging
+
+  * `unittest` integration
+
+  * 120 column rulers and corresponding `pylint` override.
