Switchboard: Clean up logging.

- Fix Switchboard log file naming to human-readable YYYY-mm-dd_HH-MM-SS instead of floating-point seconds since epoch.
- Ensure log file contains all messages (by attaching to root logger), and prepend timestamps and severity and logger name.
- Don't unnecessarily monkey-patch the `logging` module.
- Properly escape &lt; in log messages to prevent it from being misparsed as an HTML tag opening.
- Prepend logger name to the Qt output for any logger other than the default.
- Rename HTMLogger to SBLogger and clarify its purpose, removing several unnecessary passthrough methods with incorrect signatures.

diff --git a/Engine/Plugins/VirtualProduction/Switchboard/Source/Switchboard/switchboard/switchboard_dialog.py b/Engine/Plugins/VirtualProduction/Switchboard/Source/Switchboard/switchboard/switchboard_dialog.py
--- a/Engine/Plugins/VirtualProduction/Switchboard/Source/Switchboard/switchboard/switchboard_dialog.py
+++ b/Engine/Plugins/VirtualProduction/Switchboard/Source/Switchboard/switchboard/switchboard_dialog.py
@@ -27,6 +27,7 @@ from switchboard import p4_utils
 from switchboard import recording
 from switchboard import resources  # noqa
 from switchboard import switchboard_application
+from switchboard import switchboard_logging
 from switchboard import switchboard_utils
 from switchboard import switchboard_widgets as sb_widgets
 from switchboard import ugs_utils
@@ -2038,9 +2039,9 @@ class SwitchboardDialog(QtCore.QObject):
         value = self.window.logger_level_comboBox.currentText()
 
         if value == 'Message':
-            LOGGER.setLevel(logging.MESSAGE_LEVEL_NUM)
+            LOGGER.setLevel(switchboard_logging.SBLevels.MESSAGE)
         elif value == 'OSC':
-            LOGGER.setLevel(logging.OSC_LEVEL_NUM)
+            LOGGER.setLevel(switchboard_logging.SBLevels.OSC)
         elif value == 'Debug':
             LOGGER.setLevel(logging.DEBUG)
         else:
diff --git a/Engine/Plugins/VirtualProduction/Switchboard/Source/Switchboard/switchboard/switchboard_logging.py b/Engine/Plugins/VirtualProduction/Switchboard/Source/Switchboard/switchboard/switchboard_logging.py
--- a/Engine/Plugins/VirtualProduction/Switchboard/Source/Switchboard/switchboard/switchboard_logging.py
+++ b/Engine/Plugins/VirtualProduction/Switchboard/Source/Switchboard/switchboard/switchboard_logging.py
@@ -1,66 +1,78 @@
 # Copyright Epic Games, Inc. All Rights Reserved.
+
+import datetime
+from enum import IntEnum
+from io import StringIO
+import logging
 import os
 import sys
-import time
-import logging
 import tempfile
-import calendar
-import datetime
+from typing import cast, Optional
 
 from PySide6 import QtCore
 
-logging.MESSAGE_LEVEL_NUM = logging.DEBUG - 2
-logging.OSC_LEVEL_NUM = logging.DEBUG - 1
-logging.SUCCESS_LEVEL_NUM = logging.INFO + 2
+
+DEFAULT_LOGGER_NAME = 'switchboard'
+LOGGING_FMT = '%(asctime)s [%(levelname)-8s] <%(name)s>: %(message)s'
+
+
+class SBLevels(IntEnum):
+    MESSAGE = logging.DEBUG - 2
+    OSC = logging.DEBUG - 1
+    SUCCESS = logging.INFO + 2
+
 
 class QtHandler(logging.Handler):
     def __init__(self):
-        logging.Handler.__init__(self)
-        self.record = None
-
-        html = """<span style="margin: 0px; display: block"><font color="grey">[{}][{}]:</font> 
-        <font color="{}">{}</font></span>"""
-        self.html_format = html
-
-    def emit(self, record):
-        self.record = self.format(record)
-
-        if record:
-            ConsoleStream.stdout().write('{}'.format(self.record))
-            ConsoleStream.stderr().write('{}'.format(self.record))
-
-    def format(self, record):
-        levelName = record.levelname
-
-        if levelName == 'DEBUG':
-            initial = 'D'
-            color = '#66D9EF'
-        elif levelName == 'INFO':
-            initial = 'I'
-            color = 'white'
-        elif levelName == 'WARNING':
-            initial = 'W'
-            color = 'yellow' #E6DB74
-        elif levelName == 'CRITICAL':
-            initial = 'C'
-            color = '#FD971F'
-        elif levelName == 'ERROR':
-            initial = 'D'
-            color = '#F92672'
-        elif levelName == 'OSC':
-            initial = 'O'
-            color = '#4F86C6'
-        elif levelName == 'MESSAGE':
-            initial = 'M'
-            color = '#7b92ad'
-        elif levelName == 'SUCCESS':
-            initial = 'S'
-            color = '#A6E22E'
-
-        return self.html_format.format(datetime.datetime.now().strftime("%H:%M:%S"), initial, color, record.message)
-
-    def get_record(self):
-        return self.record
+        super().__init__()
+
+    def emit(self, record: logging.LogRecord):
+        html_record = self.format(record)
+        ConsoleStream.stdout().write(html_record)
+        ConsoleStream.stderr().write(html_record)
+
+    def format(self, record: logging.LogRecord):
+        match record.levelno:
+            case logging.DEBUG:
+                initial = 'D'
+                color = '#66D9EF'
+            case logging.INFO:
+                initial = 'I'
+                color = 'white'
+            case logging.WARNING:
+                initial = 'W'
+                color = 'yellow'  # E6DB74
+            case logging.CRITICAL:
+                initial = 'C'
+                color = '#FD971F'
+            case logging.ERROR:
+                initial = 'D'
+                color = '#F92672'
+            case SBLevels.OSC:
+                initial = 'O'
+                color = '#4F86C6'
+            case SBLevels.MESSAGE:
+                initial = 'M'
+                color = '#7b92ad'
+            case SBLevels.SUCCESS:
+                initial = 'S'
+                color = '#A6E22E'
+            case _:
+                initial = '?'
+                color = '#ff00ff'
+
+        time_str = datetime.datetime.now().strftime('%H:%M:%S')
+        msg = record.getMessage().replace('<', '&lt;')  # HTML escape
+        logger_name = f" &lt;{record.name}>" if record.name != DEFAULT_LOGGER_NAME else ""
+
+        return f"""
+        <span style="margin: 0px; display: block">
+            <font color="grey">
+                [{time_str}] [{initial}]{logger_name}:
+            </font>
+            <font color="{color}">{msg}</font>
+        </span>
+        """
 
 
 class ConsoleStream(QtCore.QObject):
@@ -90,121 +102,69 @@ class ConsoleStream(QtCore.QObject):
         pass
 
 
-class HTMLogger(object):
-    """This overrides the logging verbosity levels with an HTML layer for console coloring."""
-
-    def __init__(self, logger):
-        super(HTMLogger, self).__init__()
-        self.logger = logger
-        self.file_handler = None
-        self.log_path = None
-
-        self.add_custom_levels()
-
-    def add_custom_levels(self):
-        logging.addLevelName(logging.OSC_LEVEL_NUM, "OSC")
-
-        def __log_osc(s, message, *args, **kwargs):
-            if s.isEnabledFor(logging.OSC_LEVEL_NUM):
-                s._log(logging.OSC_LEVEL_NUM, message, args, **kwargs)
-        logging.Logger.osc = __log_osc
-
-
-        logging.addLevelName(logging.SUCCESS_LEVEL_NUM, "SUCCESS")
-
-        def __log_success(s, message, *args, **kwargs):
-            if s.isEnabledFor(logging.SUCCESS_LEVEL_NUM):
-                s._log(logging.SUCCESS_LEVEL_NUM, message, args, **kwargs)
-        logging.Logger.success = __log_success
+class SBLogger(logging.getLoggerClass()):
+    """ Adds convenience functions for our custom verbosity levels. """
 
+    def __init__(self, name: str, level: int = logging.NOTSET):
+        super().__init__(name, level)
 
-        logging.addLevelName(logging.MESSAGE_LEVEL_NUM, "MESSAGE")
+        self.file_handler: Optional[logging.FileHandler] = None
+        self.log_path: Optional[str] = None
 
-        def __log_message(s, message, *args, **kwargs):
-            if s.isEnabledFor(logging.MESSAGE_LEVEL_NUM):
-                s._log(logging.MESSAGE_LEVEL_NUM, message, args, **kwargs)
-        logging.Logger.message = __log_message
+    def success(self, message, *args, **kwargs):
+        self.log(SBLevels.SUCCESS, message, *args, **kwargs)
 
-    def debug(self, message, exc_info=False, current_time_only=True):
-        self.logger.debug(message, exc_info=exc_info)
+    def osc(self, message, *args, **kwargs):
+        self.log(SBLevels.OSC, message, *args, **kwargs)
 
-    def success(self, message, exc_info=False, current_time_only=True):
-        self.logger.success(message, exc_info=exc_info)
+    def message(self, message, *args, **kwargs):
+        self.log(SBLevels.MESSAGE, message, *args, **kwargs)
 
-    def info(self, message, exc_info=False, current_time_only=True):
-        self.logger.info(message, exc_info=exc_info)
-
-    def warning(self, message, exc_info=False, current_time_only=True):
-        self.logger.warning(message, exc_info=exc_info)
-
-    def critical(self, message, exc_info=False, current_time_only=True):
-        self.logger.critical(message, exc_info=exc_info)
-
-    def error(self, message, exc_info=False, current_time_only=True):
-        self.logger.error(message, exc_info=exc_info)
-
-    def osc(self, message, exc_info=False, current_time_only=True):
-        self.logger.osc(message, exc_info=exc_info)
-
-    def message(self, message, exc_info=False, current_time_only=True):
-        self.logger.message(message, exc_info=exc_info)
-
-    def get_timestamp(self, current_time_only=True):
-        current_time = self.get_current_time()
-        if current_time_only:
-            return current_time
-
-        weekday = list(calendar.day_abbr)[int(datetime.date.today().weekday())]
-        day = str(datetime.date.today()).replace("-", "/")
-        return " ".join([weekday, day, current_time])
-
-    @staticmethod
-    def get_current_time():
-        return datetime.datetime.now().strftime("%H:%M:%S")
-
-    def enable_file_logging(self, log_path=None):
+    def enable_file_logging(self, log_path: Optional[str] = None):
         if not log_path:
+            DEFAULT_LOG_PATH = os.path.join(tempfile.gettempdir(), 'switchboard')
             if not os.path.isdir(DEFAULT_LOG_PATH):
                 os.makedirs(DEFAULT_LOG_PATH)
-            self.log_path = os.path.join(DEFAULT_LOG_PATH, "switchboard_{}.html".format(str(time.time())))
+            datetime_suffix = datetime.datetime.now().strftime('%Y-%m-%d_%H-%M-%S')
+            self.log_path = os.path.join(DEFAULT_LOG_PATH, f"switchboard_{datetime_suffix}.log")
         else:
             self.log_path = log_path
         self.file_handler = logging.FileHandler(self.log_path)
-        self.logger.addHandler(self.file_handler)
-        self.logger.setLevel(logging.DEBUG)  # Log everything.
-
-    def setLevel(self, value):
-        self.logger.setLevel(value)
+        self.file_handler.setFormatter(logging.Formatter(fmt=LOGGING_FMT))
+        self.file_handler.setLevel(logging.DEBUG)  # Log everything.
+        logging.getLogger().addHandler(self.file_handler)  # Add to root.
 
     def disable_file_logging(self):
         self.log_path = None
         if not self.file_handler:
-            return self.logger.warning("No file handler found!")
-        self.logger.removeHandler(self.file_handler)
+            return self.warning("No file handler found!")
+        self.removeHandler(self.file_handler)
         self.file_handler.close()
-        
-    def save_log_file(self) -> str:
+
+    def save_log_file(self) -> Optional[str]:
         if self.log_path is not None:
-            replaced_path_name = self.log_path.replace(".html", ".log")
-            with open(replaced_path_name, "w") as file:
+            with open(self.log_path, 'w', encoding='utf-8') as file:
                 file.write(_LOG_STREAM.getvalue())
-            return replaced_path_name
+            return self.log_path
 
-from io import StringIO
-_LOG_STREAM = StringIO()
-logging.basicConfig(stream=_LOG_STREAM)
-_LOGGER = logging.getLogger()
-DEFAULT_LOG_PATH = os.path.join(tempfile.gettempdir(), "switchboard")
 
-QT_HANDLER = QtHandler()
-PYTHON_HANDLER = logging.StreamHandler()
+logging.basicConfig(level=logging.DEBUG, format=LOGGING_FMT)
+logging.addLevelName(SBLevels.OSC, 'OSC')
+logging.addLevelName(SBLevels.SUCCESS, 'SUCCESS')
+logging.addLevelName(SBLevels.MESSAGE, 'MESSAGE')
+
+logging.setLoggerClass(SBLogger)
+LOGGER = cast(SBLogger, logging.getLogger(DEFAULT_LOGGER_NAME))
 
-formatter = logging.Formatter('[%(levelname)s]: %(message)s')
-PYTHON_HANDLER.setFormatter(formatter)
+# Add our handlers to the root
+ROOT_LOGGER = logging.getLogger()
 
-_LOGGER.addHandler(QT_HANDLER)
-_LOGGER.addHandler(PYTHON_HANDLER)
-_LOGGER.setLevel(logging.DEBUG)
-PYTHON_HANDLER.setLevel(logging.DEBUG)
+QT_HANDLER = QtHandler()
+ROOT_LOGGER.addHandler(QT_HANDLER)
 
-LOGGER = HTMLogger(_LOGGER)
+# FIXME?: This seems kind of redundant.
+_LOG_STREAM = StringIO()
+LOG_STREAM_HANDLER = logging.StreamHandler(_LOG_STREAM)
+LOG_STREAM_HANDLER.setFormatter(logging.Formatter(fmt=LOGGING_FMT))
+LOG_STREAM_HANDLER.setLevel(logging.DEBUG)
+ROOT_LOGGER.addHandler(LOG_STREAM_HANDLER)
