Switchboard: Fix for `RuntimeError: Internal C++ object already deleted`.

Caused by dangling Shiboken references to widgets destroyed by their owning dialog.

#rnx

diff --git a/Engine/Plugins/VirtualProduction/Switchboard/Source/Switchboard/switchboard/config.py b/Engine/Plugins/VirtualProduction/Switchboard/Source/Switchboard/switchboard/config.py
--- a/Engine/Plugins/VirtualProduction/Switchboard/Source/Switchboard/switchboard/config.py
+++ b/Engine/Plugins/VirtualProduction/Switchboard/Source/Switchboard/switchboard/config.py
@@ -165,7 +165,7 @@ class SettingsBatchManagerToCreateUI(QtCore.QObject):
 
         super().__init__()
 
-        self.pending_settings: List['Setting'] = []
+        self.pending_settings: List[BatchItemToCreateUI] = []
         self.max_milliseconds_per_batch = max_milliseconds_per_batch
         self.batch_timer = QtCore.QTimer()
         self.batch_timer.timeout.connect(self.process_next_batch)
@@ -619,9 +619,13 @@ class Setting(QtCore.QObject):
 
         top_level_widget.destroyed.connect(on_widget_destroyed_weak)
 
-    def _on_widget_destroyed(self, destroyed_object=None, override_device_name: str = None):
+    def _on_widget_destroyed(self, destroyed_object=None, override_device_name: str | None = None):
         ''' Called when the widget associated with the given override_device_name is destroyed. '''
         self.set_widget(widget=None, override_device_name=override_device_name)
+        if override_device_name is None:
+            self._base_reset_widget = None
+        else:
+            self._reset_override_widgets.pop(override_device_name, None)
 
     def _decorate_with_reset_widget(self, override_device_name: str, setting_editor_widget: QtWidgets.QWidget):
         # Reset will still be shown on overrides
