Switchboard: Warn if device workspace Build.version stream differs from Switchboard's.

Other changes:
- More gracefully handle the corner case where roles_filename is unset, which could otherwise cause build version checks to fail.
- Log unhandled exceptions, including to the UI.
- Changing the log pane verbosity dropdown only changes the log level of the Qt UI handler, and not the underlying logger (which could result in details being omitted from log files).

diff --git a/Engine/Plugins/VirtualProduction/Switchboard/Source/Switchboard/switchboard/__main__.py b/Engine/Plugins/VirtualProduction/Switchboard/Source/Switchboard/switchboard/__main__.py
--- a/Engine/Plugins/VirtualProduction/Switchboard/Source/Switchboard/switchboard/__main__.py
+++ b/Engine/Plugins/VirtualProduction/Switchboard/Source/Switchboard/switchboard/__main__.py
@@ -41,7 +41,7 @@ def launch():
     # Global variables can be inited after QT has been started (they may show dialogues if there are errors)
     SETTINGS.init()
     CONFIG.init(SETTINGS.CONFIG)
-    
+
     # parse arguments
     args = parse_arguments()
 
@@ -81,6 +81,13 @@ def launch():
     # Logging start.
     LOGGER.info('----==== Switchboard ====----')
 
+    # Log unhandled exceptions.
+    def except_hook(exc_type: type[BaseException], exc_value: BaseException, exc_traceback):
+        exc_info=(exc_type, exc_value, exc_traceback)
+        LOGGER.critical('Unhandled exception', exc_info=exc_info)
+
+    sys.excepthook = except_hook
+
     # this will pump the event loop every 200ms so we can react faster on a SIGINT.
     # otherwise it will take several seconds before sigint_handler is called.
     timer = QtCore.QTimer()
diff --git a/Engine/Plugins/VirtualProduction/Switchboard/Source/Switchboard/switchboard/devices/unreal/plugin_unreal.py b/Engine/Plugins/VirtualProduction/Switchboard/Source/Switchboard/switchboard/devices/unreal/plugin_unreal.py
--- a/Engine/Plugins/VirtualProduction/Switchboard/Source/Switchboard/switchboard/devices/unreal/plugin_unreal.py
+++ b/Engine/Plugins/VirtualProduction/Switchboard/Source/Switchboard/switchboard/devices/unreal/plugin_unreal.py
@@ -1523,12 +1523,14 @@ class DeviceUnreal(Device):
         self.send_osc_message(osc.RECORD_STOP, 1)
 
     def _request_roles_file(self):
-        uproject_path = CONFIG.UPROJECT_PATH.get_value(
-            self.name).replace('"', '')
-        roles_filename = DeviceUnreal.csettings["roles_filename"].get_value(
-            self.name)
-        roles_file_path = os.path.join(
-            os.path.dirname(uproject_path), "Config", "Tags", roles_filename)
+        roles_filename = DeviceUnreal.csettings["roles_filename"].get_value(self.name)
+        if not roles_filename:
+            LOGGER.warning(f'{self.name}: Roles filename is unset; not requesting')
+            return
+
+        uproject_path: str = CONFIG.UPROJECT_PATH.get_value(self.name).replace('"', '')
+        uproject_dir = os.path.dirname(uproject_path)
+        roles_file_path = os.path.join(uproject_dir, "Config", "Tags", roles_filename)
         _, msg = message_protocol.create_copy_file_from_listener_message(roles_file_path)
         self.unreal_client.send_message(msg)
 
@@ -2746,10 +2748,14 @@ class DeviceUnreal(Device):
                 f"Error was '{message['error']}''")
 
     def on_file_received(self, source_path, content):
-        if source_path.endswith(DeviceUnreal.csettings["roles_filename"].get_value(self.name)):
-            self._on_receive_roles(content)
-        elif source_path.endswith('.version') or source_path.endswith('.target'):
+        if source_path.endswith('.version') or source_path.endswith('.target'):
             self._on_receive_version_file(source_path, content)
+            return
+
+        roles_filename = DeviceUnreal.csettings["roles_filename"].get_value(self.name)
+        if roles_filename and source_path.endswith(roles_filename):
+            self._on_receive_roles(content)
+            return
 
     def _on_receive_roles(self, content):
         decoded_content = base64.b64decode(content).decode()
@@ -2784,6 +2790,14 @@ class DeviceUnreal(Device):
             if source_path.endswith('Build.version'):
                 device._synced_changelist = changelist
                 device._synced_compatible_cl = compatible_cl
+
+                # Warn the user if the remote workspace points to a different branch
+                sb_buildver = ugs_utils.get_sb_buildver()
+                local_branch = sb_buildver.get('BranchName')
+                remote_branch = data.get('BranchName')
+                if local_branch is None or local_branch != remote_branch:
+                    LOGGER.warn(f'{self.name}: Workspace BranchName differs! ({remote_branch})')
+
             elif source_path.endswith('UnrealEditor.version'):
                 device._built_eng_changelist = changelist
                 device._built_eng_compatible_cl = compatible_cl
@@ -2794,7 +2808,8 @@ class DeviceUnreal(Device):
         self.device_qt_handler.signal_device_built_changelist_changed.emit(self)
 
     def on_file_receive_failed(self, source_path, error):
-        if source_path.endswith(DeviceUnreal.csettings["roles_filename"].get_value(self.name)):
+        roles_filename = DeviceUnreal.csettings["roles_filename"].get_value(self.name)
+        if roles_filename and source_path.endswith(roles_filename):
             roles = self.setting_roles.get_value()
             if len(roles) > 0:
                 LOGGER.error(
diff --git a/Engine/Plugins/VirtualProduction/Switchboard/Source/Switchboard/switchboard/switchboard_dialog.py b/Engine/Plugins/VirtualProduction/Switchboard/Source/Switchboard/switchboard/switchboard_dialog.py
--- a/Engine/Plugins/VirtualProduction/Switchboard/Source/Switchboard/switchboard/switchboard_dialog.py
+++ b/Engine/Plugins/VirtualProduction/Switchboard/Source/Switchboard/switchboard/switchboard_dialog.py
@@ -10,7 +10,6 @@ import threading
 import time
 import re
 import sys
-import json
 from typing import Callable, List, Optional, Set, Union
 
 from pathlib import Path
@@ -1065,30 +1064,16 @@ class SwitchboardDialog(QtCore.QObject):
 
         title = ['Switchboard']
 
-        # The objective is to show the UE version associated with Switchboard,
-        # which is not guaranteed to be the same as the engine_dir settings.
-        # We will rely on the location of this application, and get the associated
-        # Build.version file to get the version. An alternative to this could be to
-        # query it from the UnrealEditor executable (not SBL due to possible mis-redeployment).
+        buildver_data = ugs_utils.get_sb_buildver()
+        buildver = [
+            (f"{buildver_data.get('MajorVersion', '?')}"
+             f".{buildver_data.get('MinorVersion', '?')}"
+             f".{buildver_data.get('PatchVersion', '?')}"),
+            str(buildver_data.get('Changelist', '(unknown CL)')),
+            str(buildver_data.get('BranchName', '(unknown branch)'))
+        ]
 
-        sbmod_dir = Path(os.path.dirname(sys.modules['switchboard'].__file__))
-
-        for parent in sbmod_dir.parents:
-            if parent.name == "Engine":
-                buildver_path = parent / "Build" / "Build.version"
-                try:
-                    buildver_data = json.load(open(buildver_path, "r"))
-                    buildver = [
-                        f"{buildver_data['MajorVersion']}.{buildver_data['MinorVersion']}.{buildver_data['PatchVersion']}",
-                        str(buildver_data['Changelist']),
-                        str(buildver_data['BranchName'])
-                    ]
-
-                    title.append(f" (Unreal Engine {'-'.join(buildver)})")
-                except Exception as e:
-                    LOGGER.warning(f"Error while parsing {buildver_path}: {e}")
-
-                break
+        title.append(f" (Unreal Engine {'-'.join(buildver)})")
 
         project_name = CONFIG.PROJECT_NAME.get_value()
 
@@ -2039,13 +2024,13 @@ class SwitchboardDialog(QtCore.QObject):
         value = self.window.logger_level_comboBox.currentText()
 
         if value == 'Message':
-            LOGGER.setLevel(switchboard_logging.SBLevels.MESSAGE)
+            switchboard_logging.QT_HANDLER.setLevel(switchboard_logging.SBLevels.MESSAGE)
         elif value == 'OSC':
-            LOGGER.setLevel(switchboard_logging.SBLevels.OSC)
+            switchboard_logging.QT_HANDLER.setLevel(switchboard_logging.SBLevels.OSC)
         elif value == 'Debug':
-            LOGGER.setLevel(logging.DEBUG)
+            switchboard_logging.QT_HANDLER.setLevel(logging.DEBUG)
         else:
-            LOGGER.setLevel(logging.INFO)
+            switchboard_logging.QT_HANDLER.setLevel(logging.INFO)
 
     def logger_autoscroll_stateChanged(self, value):
         if QtCore.Qt.CheckState(value) == QtCore.Qt.Checked:
diff --git a/Engine/Plugins/VirtualProduction/Switchboard/Source/Switchboard/switchboard/switchboard_logging.py b/Engine/Plugins/VirtualProduction/Switchboard/Source/Switchboard/switchboard/switchboard_logging.py
--- a/Engine/Plugins/VirtualProduction/Switchboard/Source/Switchboard/switchboard/switchboard_logging.py
+++ b/Engine/Plugins/VirtualProduction/Switchboard/Source/Switchboard/switchboard/switchboard_logging.py
@@ -5,6 +5,7 @@ from enum import IntEnum
 from io import StringIO
 import logging
 import os
+import re
 import sys
 import tempfile
 from typing import cast, Optional
@@ -23,8 +24,11 @@ class SBLevels(IntEnum):
 
 
 class QtHandler(logging.Handler):
+    LEADING_SPACES_RE = re.compile(r'^( +)', re.MULTILINE)
+
     def __init__(self):
         super().__init__()
+        self.formatter = logging.Formatter()
 
     def emit(self, record: logging.LogRecord):
         html_record = self.format(record)
@@ -62,7 +66,13 @@ class QtHandler(logging.Handler):
                 color = '#ff00ff'
 
         time_str = datetime.datetime.now().strftime('%H:%M:%S')
-        msg = record.getMessage().replace('<', '&lt;')  # HTML escape
+        msg = self.formatter.format(record) if self.formatter else record.getMessage()
+
+        # HTML escapes
+        msg = re.sub(self.LEADING_SPACES_RE, lambda m: '&nbsp;' * len(m.group(1)), msg)
+        msg = msg.replace('<', '&lt;')
+        msg = msg.replace('\n', '<br>')
+
         logger_name = f" &lt;{record.name}>" if record.name != DEFAULT_LOGGER_NAME else ""
 
         return f"""
diff --git a/Engine/Plugins/VirtualProduction/Switchboard/Source/Switchboard/switchboard/ugs_utils.py b/Engine/Plugins/VirtualProduction/Switchboard/Source/Switchboard/switchboard/ugs_utils.py
--- a/Engine/Plugins/VirtualProduction/Switchboard/Source/Switchboard/switchboard/ugs_utils.py
+++ b/Engine/Plugins/VirtualProduction/Switchboard/Source/Switchboard/switchboard/ugs_utils.py
@@ -2,6 +2,7 @@
 
 import csv
 import io
+import json
 import logging
 import os
 import pathlib
@@ -21,6 +22,37 @@ except ImportError:
     P4PYTHON_AVAILABLE = False
 
 
+_CACHED_SB_BUILDVER: dict[str, int | str] | None = None
+
+def get_sb_buildver(*, force_refresh=False):
+    ''' Return the JSON dict of Build.version in the Switchboard source tree. '''
+    global _CACHED_SB_BUILDVER
+
+    if _CACHED_SB_BUILDVER is None or force_refresh:
+        # The objective is to query the UE version associated with Switchboard,
+        # which is not guaranteed to be the same as the engine_dir setting.
+        # We will rely on the location of this application, and get the associated
+        # Build.version file to get the version. An alternative to this could be to
+        # query it from the UnrealEditor executable (not SBL due to possible mis-redeployment).
+        sbmod_dir = pathlib.Path(os.path.dirname(sys.modules['switchboard'].__file__))
+
+        for parent in sbmod_dir.parents:
+            if parent.name == 'Engine':
+                buildver_path = parent / 'Build' / 'Build.version'
+                try:
+                    with open(buildver_path, 'r', encoding='utf-8') as f:
+                        _CACHED_SB_BUILDVER = json.load(f)
+                except Exception as e:
+                    logging.warning(f'Error while reading {buildver_path}: {e}', exc_info=True)
+
+                break
+
+    if _CACHED_SB_BUILDVER is None:
+        _CACHED_SB_BUILDVER = {}
+
+    return _CACHED_SB_BUILDVER
+
+
 def find_bin(
     dir_to_look_under: Optional[pathlib.Path] = None
 ) -> Optional[pathlib.Path]:
