Switchboard: Support loading configs from alternative paths specified via environment variable.

The code will now check two environment variables, in this order:
- Assuming Build.version can be read, and specifies a valid BranchName, it will be used to determine the first variable, e.g. "SWITCHBOARD_CONFIG_DIR_UE5_MAIN".
- If the branch-specific var cannot be determined, or is not set, it will fall back to checking "SWITCHBOARD_CONFIG_DIR" next.

If neither environment variable is set, the code should continue to behave as it did previously.

This CL also includes more logging cleanup, and some linter fixes (shadowed variables, text file open without specifying encoding, etc).

diff --git a/Engine/Plugins/VirtualProduction/Switchboard/Source/Switchboard/switchboard/__main__.py b/Engine/Plugins/VirtualProduction/Switchboard/Source/Switchboard/switchboard/__main__.py
--- a/Engine/Plugins/VirtualProduction/Switchboard/Source/Switchboard/switchboard/__main__.py
+++ b/Engine/Plugins/VirtualProduction/Switchboard/Source/Switchboard/switchboard/__main__.py
@@ -9,7 +9,7 @@ from PySide6 import QtCore, QtWidgets
 from switchboard.config import SETTINGS, CONFIG
 from switchboard.switchboard_scripting import ScriptManager
 from .switchboard_dialog import SwitchboardDialog
-from .switchboard_logging import LOGGER
+from .switchboard_logging import LOGGER, FILE_HANDLER
 
 
 # Build resources
@@ -75,11 +75,9 @@ def launch():
 
     main_window.window.show()
 
-    # Enable file logging.
-    LOGGER.enable_file_logging()
-
     # Logging start.
     LOGGER.info('----==== Switchboard ====----')
+    LOGGER.info(f'Log file: {FILE_HANDLER.baseFilename}')
 
     # Log unhandled exceptions.
     def except_hook(exc_type: type[BaseException], exc_value: BaseException, exc_traceback):
diff --git a/Engine/Plugins/VirtualProduction/Switchboard/Source/Switchboard/switchboard/add_config_dialog.py b/Engine/Plugins/VirtualProduction/Switchboard/Source/Switchboard/switchboard/add_config_dialog.py
--- a/Engine/Plugins/VirtualProduction/Switchboard/Source/Switchboard/switchboard/add_config_dialog.py
+++ b/Engine/Plugins/VirtualProduction/Switchboard/Source/Switchboard/switchboard/add_config_dialog.py
@@ -202,7 +202,7 @@ class AddConfigDialog(QtWidgets.QDialog):
                 return
 
         try:
-            config_path = config.get_absolute_config_path(config_path_str)
+            config_path = config.CONFIG_MGR.get_absolute_config_path(config_path_str)
         except Exception as e:
             error_msg = QtWidgets.QErrorMessage(self)
             error_msg.showMessage(str(e))
@@ -534,7 +534,7 @@ class AddConfigDialog(QtWidgets.QDialog):
 
     def on_browse_config_path(self):
         config_path_str, _ = QtWidgets.QFileDialog.getSaveFileName(self,
-            'Select config file', str(config.ROOT_CONFIGS_PATH),
+            'Select config file', str(config.CONFIG_MGR.config_dir),
             f'Config files (*{config.CONFIG_SUFFIX})')
 
         if not config_path_str:
@@ -544,9 +544,9 @@ class AddConfigDialog(QtWidgets.QDialog):
         #
         try:
             config_path_str = str(
-                config.get_relative_config_path(config_path_str))
+                config.CONFIG_MGR.get_relative_config_path(config_path_str))
         except Exception as e:
-            config_path_str = str(config.get_absolute_config_path(config_path_str))
+            config_path_str = str(config.CONFIG_MGR.get_absolute_config_path(config_path_str))
 
         self.config_path_line_edit.setText(config_path_str)
 
diff --git a/Engine/Plugins/VirtualProduction/Switchboard/Source/Switchboard/switchboard/config.py b/Engine/Plugins/VirtualProduction/Switchboard/Source/Switchboard/switchboard/config.py
--- a/Engine/Plugins/VirtualProduction/Switchboard/Source/Switchboard/switchboard/config.py
+++ b/Engine/Plugins/VirtualProduction/Switchboard/Source/Switchboard/switchboard/config.py
@@ -15,7 +15,7 @@ import weakref
 
 from enum import Enum
 from pathlib import Path, PurePosixPath
-from typing import (Any, Callable, Optional, Tuple, Type, Union, List)
+from typing import Any, Callable, Optional, Tuple, Type, Union, List
 
 from PySide6 import QtCore
 from PySide6 import QtGui
@@ -23,25 +23,145 @@ from PySide6 import QtWidgets
 
 from switchboard import switchboard_widgets as sb_widgets
 from switchboard.switchboard_logging import LOGGER
-from switchboard.switchboard_widgets import (
-    DropDownMenuComboBox, NonScrollableComboBox)
+from switchboard.switchboard_widgets import DropDownMenuComboBox, NonScrollableComboBox
 from switchboard import ugs_utils
 from switchboard.sbcache import SBCache, Map
 from switchboard.ue_plugin_utils import UnrealPlugin, UnrealPluginManager
 
-ROOT_CONFIGS_PATH = pathlib.Path(__file__).parent.with_name('configs')
-CONFIG_SUFFIX = '.json'
 
-USER_SETTINGS_FILE_NAME = 'user_settings.json'
-USER_SETTINGS_FILE_PATH = ROOT_CONFIGS_PATH.joinpath(USER_SETTINGS_FILE_NAME)
-USER_SETTINGS_BACKUP_FILE_NAME = 'corrupted_user_settings_backup.json'
-USER_SETTINGS_BACKUP_FILE_PATH = ROOT_CONFIGS_PATH.joinpath(USER_SETTINGS_BACKUP_FILE_NAME)
+CONFIG_SUFFIX = '.json'
 
 DEFAULT_MAP_TEXT = '-- Default Map --'
 
 ENABLE_UGS_SUPPORT = False
 
 
+class ConfigManager:
+    DEFAULT_CONFIG_DIR = pathlib.Path(__file__).parent.with_name('configs')
+    CONFIG_ENV_VAR = 'SWITCHBOARD_CONFIG_DIR'
+
+    USER_SETTINGS_FILE_NAME = 'user_settings.json'
+    USER_SETTINGS_BACKUP_FILE_NAME = 'corrupted_user_settings_backup.json'
+
+    def __init__(
+        self,
+        config_dir: pathlib.Path | None = None
+    ):
+        LOGGER.debug('Initializing ConfigManager')
+
+        self.config_dir = config_dir or self._determine_config_dir()
+        LOGGER.debug(f'Config dir set to {self.config_dir}')
+
+    @classmethod
+    def _fmt_env_var(cls, var) -> str:
+        return f'%{var}%' if sys.platform.startswith('win') else f'${var}'
+
+    def _config_dir_valid(self, check_dir: pathlib.Path) -> bool:
+        check_dir.mkdir(parents=True, exist_ok=True)
+        return check_dir.is_dir()
+
+    def _determine_config_dir(self) -> pathlib.Path:
+        sb_buildver = ugs_utils.get_sb_buildver()
+        if sb_branch := sb_buildver.get('BranchName'): # "++UE5+Main"
+            sb_branch = str(sb_branch).strip('+').replace('+', '_').upper() # "UE5_MAIN"
+            stream_env_var = f'{self.CONFIG_ENV_VAR}_{sb_branch}' # "SB_CONFIG_DIR_UE5_MAIN"
+
+            LOGGER.debug(f'Looking for config dir override: {self._fmt_env_var(stream_env_var)}')
+            if stream_dir := os.environ.get(stream_env_var):
+                stream_dir = pathlib.Path(stream_dir)
+                if self._config_dir_valid(stream_dir):
+                    return stream_dir
+
+        LOGGER.debug(f'Looking for config dir override: {self._fmt_env_var(self.CONFIG_ENV_VAR)}')
+        if override_dir := os.environ.get(self.CONFIG_ENV_VAR):
+            override_dir = pathlib.Path(override_dir)
+            if self._config_dir_valid(override_dir):
+                return override_dir
+
+        return self.DEFAULT_CONFIG_DIR
+
+    @property
+    def user_settings_file_path(self) -> pathlib.Path:
+        '''Path to the "global" settings, including the last loaded config.'''
+        return self.config_dir.joinpath(self.USER_SETTINGS_FILE_NAME)
+
+    @property
+    def user_settings_backup_file_path(self) -> pathlib.Path:
+        '''Path at which a backup is automatically created if we're unable to load.'''
+        return self.config_dir.joinpath(self.USER_SETTINGS_BACKUP_FILE_NAME)
+
+    def list_config_paths(self) -> list[pathlib.Path]:
+        '''
+        Returns a list of absolute paths to all config files in the configs dir.
+        '''
+        self.config_dir.mkdir(parents=True, exist_ok=True)
+
+        # Find all JSON files in the config dir recursively, but exclude the user
+        # settings file.
+        config_paths = [
+            path for path in self.config_dir.rglob(f'*{CONFIG_SUFFIX}')
+            if (path != self.user_settings_file_path and
+                path != self.user_settings_backup_file_path)
+        ]
+
+        return config_paths
+
+    def get_absolute_config_path(
+        self,
+        config_path: Union[str, pathlib.Path]
+    ) -> pathlib.Path:
+        '''
+        Returns the given string or path object as an absolute config path.
+
+        The string/path is validated to ensure that:
+        - It is not empty, or all whitespace
+        - It ends with the config path suffix
+        - It is not the same path as the user settings file path
+        '''
+        if isinstance(config_path, str):
+            config_path = config_path.strip()
+            if not config_path:
+                raise ConfigPathEmptyError('Config path cannot be empty')
+
+            config_path = pathlib.Path(config_path)
+
+        # Manually add the suffix instead of using pathlib.Path.with_suffix().
+        # For strings like "foo.bar", with_suffix() will first remove ".bar"
+        # before adding the suffix, which we don't want it to do.
+        if not config_path.name.endswith(CONFIG_SUFFIX):
+            config_path = config_path.with_name(
+                f'{config_path.name}{CONFIG_SUFFIX}')
+
+        if not config_path.is_absolute():
+            # Relative paths can simply be made absolute.
+            config_path = self.config_dir.joinpath(config_path)
+
+        if config_path.resolve() == self.user_settings_file_path:
+            raise ConfigPathIsUserSettingsError(
+                'Config path cannot be the same as the user settings file '
+                f'path "{self.user_settings_file_path}"')
+
+        return config_path
+
+    def get_relative_config_path(
+        self,
+        config_path: Union[str, pathlib.Path]
+    ) -> pathlib.Path:
+        '''
+        Returns the given string or path object as a config path relative to the
+        root configs path.
+
+        An absolute path is generated first to perform all of the same validation
+        as get_absolute_config_path() before the relative path is computed and
+        returned.
+        '''
+        config_path = self.get_absolute_config_path(config_path)
+        return config_path.relative_to(self.config_dir)
+
+
+CONFIG_MGR = ConfigManager()
+
+
 def map_name_is_valid(map_name: str) -> bool:
     """
     Check if map is specified
@@ -69,12 +189,13 @@ def get_game_launch_level_path() -> str:
         parser = ugs_utils.IniParser()
 
         try:
-            engine_ini_file = open(engine_ini, "r", encoding="utf-8")
-            parser.read_file(engine_ini_file)
+            with open(engine_ini, encoding='utf-8') as engine_ini_file:
+                parser.read_file(engine_ini_file)
 
-            levels = parser.try_get('/Script/EngineSettings.GameMapsSettings', 'GameDefaultMap')
+                levels = parser.try_get('/Script/EngineSettings.GameMapsSettings',
+                                        'GameDefaultMap')
 
-            if len(levels) == 0:
+            if not levels:
                 raise Exception('Default map is not set in DefaultEngine.ini')
 
             level = levels[0]
@@ -2246,14 +2367,6 @@ class ConfigPathEmptyError(ConfigPathError):
     pass
 
 
-class ConfigPathLocationError(ConfigPathError):
-    '''
-    Exception type raised when a config path is located outside of the root
-    configs directory.
-    '''
-    pass
-
-
 class ConfigPathIsUserSettingsError(ConfigPathError):
     '''
     Exception type raised when the user settings file path is used as a config
@@ -2262,56 +2375,6 @@ class ConfigPathIsUserSettingsError(ConfigPathError):
     pass
 
 
-def get_absolute_config_path(
-        config_path: Union[str, pathlib.Path]) -> pathlib.Path:
-    '''
-    Returns the given string or path object as an absolute config path.
-
-    The string/path is validated to ensure that:
-      - It is not empty, or all whitespace
-      - It ends with the config path suffix
-      - It is not the same path as the user settings file path
-    '''
-    if isinstance(config_path, str):
-        config_path = config_path.strip()
-        if not config_path:
-            raise ConfigPathEmptyError('Config path cannot be empty')
-
-        config_path = pathlib.Path(config_path)
-
-    # Manually add the suffix instead of using pathlib.Path.with_suffix().
-    # For strings like "foo.bar", with_suffix() will first remove ".bar"
-    # before adding the suffix, which we don't want it to do.
-    if not config_path.name.endswith(CONFIG_SUFFIX):
-        config_path = config_path.with_name(
-            f'{config_path.name}{CONFIG_SUFFIX}')
-
-    if not config_path.is_absolute():
-        # Relative paths can simply be made absolute.
-        config_path = ROOT_CONFIGS_PATH.joinpath(config_path)
-
-    if config_path.resolve() == USER_SETTINGS_FILE_PATH:
-        raise ConfigPathIsUserSettingsError(
-            'Config path cannot be the same as the user settings file '
-            f'path "{USER_SETTINGS_FILE_PATH}"')
-
-    return config_path
-
-
-def get_relative_config_path(
-        config_path: Union[str, pathlib.Path]) -> pathlib.Path:
-    '''
-    Returns the given string or path object as a config path relative to the
-    root configs path.
-
-    An absolute path is generated first to perform all of the same validation
-    as get_absolute_config_path() before the relative path is computed and
-    returned.
-    '''
-    config_path = get_absolute_config_path(config_path)
-    return config_path.relative_to(ROOT_CONFIGS_PATH)
-
-
 class ConfigPathValidator(QtGui.QValidator):
     '''
     Validator to determine whether the input is an acceptable config file
@@ -2324,7 +2387,7 @@ class ConfigPathValidator(QtGui.QValidator):
 
     def validate(self, input, pos):
         try:
-            get_absolute_config_path(input)
+            CONFIG_MGR.get_absolute_config_path(input)
         except Exception:
             return QtGui.QValidator.Intermediate
 
@@ -2342,7 +2405,7 @@ class EngineSyncMethod(Enum):
 
 class Config(object):
 
-    DEFAULT_CONFIG_PATH = ROOT_CONFIGS_PATH.joinpath(f'Default{CONFIG_SUFFIX}')
+    DEFAULT_CONFIG_PATH = CONFIG_MGR.get_relative_config_path(f'Default')
 
     saving_allowed = True
     saving_allowed_fifo = []
@@ -2364,10 +2427,10 @@ class Config(object):
     def init_with_file_path(self, file_path: Union[str, pathlib.Path]):
         if file_path:
             try:
-                self.file_path = get_absolute_config_path(file_path)
+                self.file_path = CONFIG_MGR.get_absolute_config_path(file_path)
 
                 # Read the json config file
-                with open(self.file_path, 'r') as f:
+                with open(self.file_path, encoding='utf-8') as f:
                     LOGGER.debug(f'Loading Config {self.file_path}')
                     data = json.load(f)
 
@@ -2377,7 +2440,7 @@ class Config(object):
                 data = {}
             except ValueError:
                 # The original file will be overwritten
-                self._backup_corrupted_config(self.file_path.__str__())
+                self._backup_corrupted_config(str(self.file_path))
                 data = {}
         else:
             self.file_path = None
@@ -2482,7 +2545,7 @@ class Config(object):
         ''' Initialize new configuration '''
 
         # Assign self.file_path before the basic_project_settings line because it requires it to be updated.
-        self.file_path = get_absolute_config_path(file_path)
+        self.file_path = CONFIG_MGR.get_absolute_config_path(file_path)
 
         basic_project_settings = {
                 "project_name": self.file_path.stem,
@@ -2877,7 +2940,7 @@ class Config(object):
 
         If a file already exists at the new path, it will be overwritten.
         """
-        new_config_path = get_absolute_config_path(new_config_path)
+        new_config_path = CONFIG_MGR.get_absolute_config_path(new_config_path)
 
         if self.file_path:
             new_config_path.parent.mkdir(parents=True, exist_ok=True)
@@ -2892,7 +2955,7 @@ class Config(object):
 
         If a file already exists at the new path, it will be overwritten.
         """
-        new_config_path = get_absolute_config_path(new_config_path)
+        new_config_path = CONFIG_MGR.get_absolute_config_path(new_config_path)
 
         new_project_name = new_config_path.stem
 
@@ -2993,7 +3056,7 @@ class Config(object):
 
         try:
             self.file_path.parent.mkdir(parents=True, exist_ok=True)
-            with open(self.file_path, 'w') as f:
+            with open(self.file_path, 'w', encoding='utf-8') as f:
                 temp_buf.seek(0)
                 shutil.copyfileobj(temp_buf, f)
                 LOGGER.debug(f'Config File: {self.file_path} updated')
@@ -3289,41 +3352,45 @@ class Config(object):
 
 class UserSettings(object):
     def init(self):
+        user_settings_path = CONFIG_MGR.user_settings_file_path
+
         try:
-            with open(USER_SETTINGS_FILE_PATH) as f:
-                LOGGER.debug(f'Loading Settings {USER_SETTINGS_FILE_PATH}')
+            with open(user_settings_path, encoding='utf-8') as f:
+                LOGGER.debug(f'Loading Settings {user_settings_path}')
                 data = json.load(f)
         except FileNotFoundError:
             # Create a default user_settings
             LOGGER.debug('Creating default user settings')
             data = {}
         except ValueError:
-            LOGGER.error(f'{USER_SETTINGS_FILE_NAME} has invalid JSON format. ')
+            LOGGER.error(f'{user_settings_path} has invalid JSON format. ')
+
+            backup_path = CONFIG_MGR.user_settings_backup_file_path
 
             answer = QtWidgets.QMessageBox.question(
                 None,
                 'Invalid User Settings',
-                 f'User settings has invalid JSON format and will be replaced with a valid a new default JSON.'
-                 f'\n\nDo you want to save a backup named {USER_SETTINGS_BACKUP_FILE_NAME} (overrides existing)?',
+                'User settings has invalid JSON format and will be replaced with a valid a new default JSON.\n\n'
+                f'Do you want to save a backup named {ConfigManager.USER_SETTINGS_BACKUP_FILE_NAME} (overwrites existing)?',
                  QtWidgets.QMessageBox.Yes | QtWidgets.QMessageBox.No
             )
             if answer == QtWidgets.QMessageBox.Yes:
-                if os.path.exists(USER_SETTINGS_BACKUP_FILE_PATH):
-                    os.remove(USER_SETTINGS_BACKUP_FILE_PATH)
-                shutil.copy(USER_SETTINGS_FILE_PATH, USER_SETTINGS_BACKUP_FILE_PATH)
+                if os.path.exists(backup_path):
+                    os.remove(backup_path)
+                shutil.copy(user_settings_path, backup_path)
 
             data = {}
 
         self.CONFIG = data.get('config')
         if self.CONFIG:
             try:
-                self.CONFIG = get_absolute_config_path(self.CONFIG)
+                self.CONFIG = CONFIG_MGR.get_absolute_config_path(self.CONFIG)
             except ConfigPathError as e:
                 LOGGER.error(e)
                 self.CONFIG = None
 
         if not self.CONFIG:
-            config_paths = list_config_paths()
+            config_paths = CONFIG_MGR.list_config_paths()
             self.CONFIG = config_paths[0] if config_paths else None
 
         # Address of the machine running Switchboard
@@ -3370,31 +3437,16 @@ class UserSettings(object):
 
         if self.CONFIG:
             try:
-                data['config'] = str(get_relative_config_path(self.CONFIG))
+                data['config'] = str(CONFIG_MGR.get_relative_config_path(self.CONFIG))
             except ConfigPathError as e:
                 LOGGER.error(e)
             except ValueError as e:
                 data['config'] = str(self.CONFIG)
 
-        with open(USER_SETTINGS_FILE_PATH, 'w') as f:
+        with open(CONFIG_MGR.user_settings_file_path, 'w', encoding='utf-8') as f:
             json.dump(data, f, indent=4)
 
 
-def list_config_paths() -> list[pathlib.Path]:
-    '''
-    Returns a list of absolute paths to all config files in the configs dir.
-    '''
-    ROOT_CONFIGS_PATH.mkdir(parents=True, exist_ok=True)
-
-    # Find all JSON files in the config dir recursively, but exclude the user
-    # settings file.
-    config_paths = [
-        path for path in ROOT_CONFIGS_PATH.rglob(f'*{CONFIG_SUFFIX}')
-        if path != USER_SETTINGS_FILE_PATH and path != USER_SETTINGS_BACKUP_FILE_PATH]
-
-    return config_paths
-
-
 # Get the user settings and load their config
 SETTINGS = UserSettings()
 CONFIG = Config()
diff --git a/Engine/Plugins/VirtualProduction/Switchboard/Source/Switchboard/switchboard/settings_dialog.py b/Engine/Plugins/VirtualProduction/Switchboard/Source/Switchboard/switchboard/settings_dialog.py
--- a/Engine/Plugins/VirtualProduction/Switchboard/Source/Switchboard/switchboard/settings_dialog.py
+++ b/Engine/Plugins/VirtualProduction/Switchboard/Source/Switchboard/switchboard/settings_dialog.py
@@ -87,10 +87,10 @@ class CreateUISearchUpdater:
 
 
 class SettingsDialog(QtCore.QObject):
-    def __init__(self, settings: UserSettings, config: Config):
+    def __init__(self, settings: UserSettings, in_config: Config):
         super().__init__()
         self.plugin_widgets = {}
-        self.config = config
+        self.config = in_config
 
         # Set the UI object
         self.ui = QtWidgets.QDialog()
@@ -103,9 +103,9 @@ class SettingsDialog(QtCore.QObject):
         self._create_search_area(dialog_layout)
         self.general_settings_list = [
             self._create_config_path_settings(),
-            self._create_switchboard_settings(settings, config),
-            self._create_project_settings(config),
-            self._create_multi_user_server_settings(config)
+            self._create_switchboard_settings(settings, self.config),
+            self._create_project_settings(self.config),
+            self._create_multi_user_server_settings(self.config)
         ]
 
         self._create_tab_widget(dialog_layout)
@@ -180,8 +180,7 @@ class SettingsDialog(QtCore.QObject):
         if (self._changed_config_path in self._config_paths
                 and self._changed_config_path != self._current_config_path):
             # Show the confirmation dialog using a relative path to the config.
-            rel_config_path = self.config.get_relative_config_path(
-                self._changed_config_path)
+            rel_config_path = config.CONFIG_MGR.get_relative_config_path(self._changed_config_path)
             reply = QtWidgets.QMessageBox.question(
                 self.ui, 'Confirm Overwrite',
                 ('Are you sure you would like to change the config path and '
@@ -266,12 +265,12 @@ class SettingsDialog(QtCore.QObject):
         layout.addWidget(self.ui.config_path_line_edit)
 
         # Store the current config paths so we can warn about overwriting an existing config.
-        self._config_paths = config.list_config_paths()
+        self._config_paths = config.CONFIG_MGR.list_config_paths()
         self.set_config_path(SETTINGS.CONFIG)
 
         return self.ui.config_path_layout
 
-    def _create_switchboard_settings(self, settings: UserSettings, config: Config):
+    def _create_switchboard_settings(self, settings: UserSettings, in_config: Config):
         self.ui.switchboard_settings_group = CollapsibleGroupBox()
         self.ui.switchboard_settings_group.setTitle("Switchboard")
         layout = QtWidgets.QVBoxLayout(self.ui.switchboard_settings_group)
@@ -279,8 +278,8 @@ class SettingsDialog(QtCore.QObject):
         settings_dict = {
             "address": settings.ADDRESS,
             "transport_path": settings.TRANSPORT_PATH,
-            "listener_exe": config.LISTENER_EXE,
-            "sblhelper_exe": config.SBLHELPER_EXE,
+            "listener_exe": in_config.LISTENER_EXE,
+            "sblhelper_exe": in_config.SBLHELPER_EXE,
         }
 
         self.ui.switchboard_settings_sub_group = QtWidgets.QGroupBox()
@@ -290,30 +289,30 @@ class SettingsDialog(QtCore.QObject):
 
         return self.ui.switchboard_settings_group
 
-    def _create_project_settings(self, config: Config):
+    def _create_project_settings(self, in_config: Config):
         self.ui.project_settings_group = CollapsibleGroupBox()
         self.ui.project_settings_group.setTitle("Project Settings")
         layout = QtWidgets.QVBoxLayout(self.ui.project_settings_group)
 
-        self._create_basic_settings(config, layout)
-        self._create_osc_settings(config, layout)
-        self._create_source_control_settings(config, layout)
+        self._create_basic_settings(in_config, layout)
+        self._create_osc_settings(in_config, layout)
+        self._create_source_control_settings(in_config, layout)
 
         return self.ui.project_settings_group
 
-    def _create_basic_settings(self, config: Config, layout: QtWidgets.QFormLayout):
+    def _create_basic_settings(self, in_config: Config, layout: QtWidgets.QFormLayout):
         self.ui.basic_settings_group = QtWidgets.QGroupBox()
         form_layout = QtWidgets.QFormLayout(self.ui.basic_settings_group)
         layout.addWidget(self.ui.basic_settings_group)
-        self._create_settings_section(config.basic_project_settings, form_layout)
+        self._create_settings_section(in_config.basic_project_settings, form_layout)
 
-    def _create_osc_settings(self, config: Config, layout: QtWidgets.QFormLayout):
+    def _create_osc_settings(self, in_config: Config, layout: QtWidgets.QFormLayout):
         self.ui.osc_settings_group = QtWidgets.QGroupBox("OSC")
         form_layout = QtWidgets.QFormLayout(self.ui.osc_settings_group)
         layout.addWidget(self.ui.osc_settings_group)
-        self._create_settings_section(config.osc_settings, form_layout)
+        self._create_settings_section(in_config.osc_settings, form_layout)
 
-    def _create_source_control_settings(self, config: Config, layout: QtWidgets.QFormLayout):
+    def _create_source_control_settings(self, in_config: Config, layout: QtWidgets.QFormLayout):
         self.ui.source_control_settings_group = QtWidgets.QGroupBox("Source Control")
         self.ui.source_control_settings_group.setCheckable(True)
         self.ui.source_control_settings_group.setStyleSheet(
@@ -326,10 +325,13 @@ class SettingsDialog(QtCore.QObject):
         layout.addWidget(self.ui.source_control_settings_group)
 
         # Source control setting UI is fully enabled / disabled depending on P4_ENABLED
-        self.ui.source_control_settings_group.setChecked(config.P4_ENABLED.get_value())
+        self.ui.source_control_settings_group.setChecked(in_config.P4_ENABLED.get_value())
         self.ui.source_control_settings_group.toggled.connect(self._on_group_source_control_changed)
-        config.P4_ENABLED.signal_setting_changed.connect(self._on_source_control_setting_changed)
-        settings_to_show = {key: setting for (key, setting) in config.source_control_settings.items() if setting != config.P4_ENABLED}
+        in_config.P4_ENABLED.signal_setting_changed.connect(self._on_source_control_setting_changed)
+        settings_to_show = {
+            key: setting for (key, setting) in in_config.source_control_settings.items()
+            if setting != in_config.P4_ENABLED
+        }
 
         self._create_settings_section(settings_to_show, form_layout)
 
@@ -340,7 +342,7 @@ class SettingsDialog(QtCore.QObject):
     def _on_source_control_setting_changed(self, old_value: bool, new_value: bool):
         self.ui.source_control_settings_group.setChecked(new_value)
 
-    def _create_multi_user_server_settings(self, config: Config):
+    def _create_multi_user_server_settings(self, in_config: Config):
         self.ui.multi_user_settings = CollapsibleGroupBox()
         self.ui.multi_user_settings.setTitle("Multi User Server")
         layout = QtWidgets.QVBoxLayout(self.ui.multi_user_settings)
@@ -348,7 +350,7 @@ class SettingsDialog(QtCore.QObject):
         self.ui.multi_user_settings_sub_group = QtWidgets.QGroupBox()
         form_layout = QtWidgets.QFormLayout(self.ui.multi_user_settings_sub_group)
         layout.addWidget(self.ui.multi_user_settings_sub_group)
-        self._create_settings_section(config.mu_settings, form_layout)
+        self._create_settings_section(in_config.mu_settings, form_layout)
 
         return self.ui.multi_user_settings
 
@@ -408,7 +410,8 @@ class SettingsDialog(QtCore.QObject):
         else:
             # prefer relative path when possible
             try:
-                config_path_str = str(config.get_relative_config_path(config_path).with_suffix(''))
+                config_path_str = str(
+                    config.CONFIG_MGR.get_relative_config_path(config_path).with_suffix(''))
             except ValueError:
                 config_path_str = str(config_path)
 
@@ -421,7 +424,7 @@ class SettingsDialog(QtCore.QObject):
             self.ui.config_path_line_edit, "input_error", False)
 
         try:
-            config_path = config.get_absolute_config_path(config_path_str)
+            config_path = config.CONFIG_MGR.get_absolute_config_path(config_path_str)
         except Exception as e:
             sb_widgets.set_qt_property(
                 self.ui.config_path_line_edit, "input_error", True)
@@ -448,7 +451,7 @@ class SettingsDialog(QtCore.QObject):
 
         all_settings = instance_settings + plugin_settings
 
-        scategories = {}
+        scategories: dict[str, list[Setting]] = {}
 
         for setting in all_settings:
 
@@ -496,8 +499,8 @@ class SettingsDialog(QtCore.QObject):
         self,
         plugin_name: str,
         plugin_cls: Device,
-        plugin_settings,
-        device_settings
+        plugin_settings: list[Setting],
+        device_settings: list[tuple[str, list[Setting], list[Setting]]]
     ):
 
         any_device_settings = (
diff --git a/Engine/Plugins/VirtualProduction/Switchboard/Source/Switchboard/switchboard/switchboard_dialog.py b/Engine/Plugins/VirtualProduction/Switchboard/Source/Switchboard/switchboard/switchboard_dialog.py
--- a/Engine/Plugins/VirtualProduction/Switchboard/Source/Switchboard/switchboard/switchboard_dialog.py
+++ b/Engine/Plugins/VirtualProduction/Switchboard/Source/Switchboard/switchboard/switchboard_dialog.py
@@ -959,14 +959,14 @@ class SwitchboardDialog(QtCore.QObject):
         # that directory as we go. Config files in the root configs directory
         # will go at the top level.
         menu_hierarchy = {
-            config.ROOT_CONFIGS_PATH: self.window.menu_load_config
+            config.CONFIG_MGR.config_dir: self.window.menu_load_config
         }
 
         def _get_menu_for_path(path):
             # Safe guard to make sure we don't accidentally traverse up past
             # the root configs path.
-            if (path != config.ROOT_CONFIGS_PATH and
-                    config.ROOT_CONFIGS_PATH not in path.parents):
+            if (path != config.CONFIG_MGR.config_dir and
+                    config.CONFIG_MGR.config_dir not in path.parents):
                 return None
 
             path_menu = menu_hierarchy.get(path)
@@ -980,7 +980,7 @@ class SwitchboardDialog(QtCore.QObject):
 
             return path_menu
 
-        config_paths = config.list_config_paths()
+        config_paths = config.CONFIG_MGR.list_config_paths()
 
         # Take a first pass through the config paths just creating the
         # submenus. This makes sure that all submenus appear before any configs
@@ -1012,7 +1012,7 @@ class SwitchboardDialog(QtCore.QObject):
         '''
         config_path, _ = QtWidgets.QFileDialog.getOpenFileName(
             self.window,
-            'Select config file', str(config.ROOT_CONFIGS_PATH),
+            'Select config file', str(config.CONFIG_MGR.config_dir),
             f'Config files (*{config.CONFIG_SUFFIX})'
         )
 
@@ -1141,7 +1141,7 @@ class SwitchboardDialog(QtCore.QObject):
 
         new_config_path, _ = QtWidgets.QFileDialog.getSaveFileName(
             self.window,
-            'Select config file', str(config.ROOT_CONFIGS_PATH),
+            'Select config file', str(config.CONFIG_MGR.config_dir),
             f'Config files (*{config.CONFIG_SUFFIX})'
         )
 
@@ -1200,7 +1200,7 @@ class SwitchboardDialog(QtCore.QObject):
         config.list_config_paths(), or it will create a new config.
         """
         # Show the confirmation dialog using a relative path to the config.
-        rel_config_path = config.get_relative_config_path(SETTINGS.CONFIG)
+        rel_config_path = config.CONFIG_MGR.get_relative_config_path(SETTINGS.CONFIG)
         reply = QtWidgets.QMessageBox.question(
             self.window,
             'Delete Config',
@@ -1216,7 +1216,7 @@ class SwitchboardDialog(QtCore.QObject):
                 LOGGER.error(f'menu_delete_config: {e}')
 
             # Grab a new config to lead once this one is deleted
-            config_paths = config.list_config_paths()
+            config_paths = config.CONFIG_MGR.list_config_paths()
 
             if config_paths:
                 self.set_current_config(config_paths[0])
diff --git a/Engine/Plugins/VirtualProduction/Switchboard/Source/Switchboard/switchboard/switchboard_logging.py b/Engine/Plugins/VirtualProduction/Switchboard/Source/Switchboard/switchboard/switchboard_logging.py
--- a/Engine/Plugins/VirtualProduction/Switchboard/Source/Switchboard/switchboard/switchboard_logging.py
+++ b/Engine/Plugins/VirtualProduction/Switchboard/Source/Switchboard/switchboard/switchboard_logging.py
@@ -130,32 +130,15 @@ class SBLogger(logging.getLoggerClass()):
     def message(self, message, *args, **kwargs):
         self.log(SBLevels.MESSAGE, message, *args, **kwargs)
 
-    def enable_file_logging(self, log_path: Optional[str] = None):
-        if not log_path:
-            DEFAULT_LOG_PATH = os.path.join(tempfile.gettempdir(), 'switchboard')
-            if not os.path.isdir(DEFAULT_LOG_PATH):
-                os.makedirs(DEFAULT_LOG_PATH)
-            datetime_suffix = datetime.datetime.now().strftime('%Y-%m-%d_%H-%M-%S')
-            self.log_path = os.path.join(DEFAULT_LOG_PATH, f"switchboard_{datetime_suffix}.log")
-        else:
-            self.log_path = log_path
-        self.file_handler = logging.FileHandler(self.log_path)
-        self.file_handler.setFormatter(logging.Formatter(fmt=LOGGING_FMT))
-        self.file_handler.setLevel(logging.DEBUG)  # Log everything.
-        logging.getLogger().addHandler(self.file_handler)  # Add to root.
-
-    def disable_file_logging(self):
-        self.log_path = None
-        if not self.file_handler:
-            return self.warning("No file handler found!")
-        self.removeHandler(self.file_handler)
-        self.file_handler.close()
-
-    def save_log_file(self) -> Optional[str]:
-        if self.log_path is not None:
-            with open(self.log_path, 'w', encoding='utf-8') as file:
-                file.write(_LOG_STREAM.getvalue())
-            return self.log_path
+    @staticmethod
+    def make_log_path(timestamp: datetime.datetime | None = None) -> str:
+        DEFAULT_LOG_DIR = os.path.join(tempfile.gettempdir(), 'switchboard')
+        if not os.path.isdir(DEFAULT_LOG_DIR):
+            os.makedirs(DEFAULT_LOG_DIR)
+        if not timestamp:
+            timestamp = datetime.datetime.now()
+        datetime_suffix = timestamp.strftime('%Y-%m-%d_%H-%M-%S')
+        return os.path.join(DEFAULT_LOG_DIR, f"switchboard_{datetime_suffix}.log")
 
 
 logging.basicConfig(level=logging.DEBUG, format=LOGGING_FMT)
@@ -172,9 +155,6 @@ ROOT_LOGGER = logging.getLogger()
 QT_HANDLER = QtHandler()
 ROOT_LOGGER.addHandler(QT_HANDLER)
 
-# FIXME?: This seems kind of redundant.
-_LOG_STREAM = StringIO()
-LOG_STREAM_HANDLER = logging.StreamHandler(_LOG_STREAM)
-LOG_STREAM_HANDLER.setFormatter(logging.Formatter(fmt=LOGGING_FMT))
-LOG_STREAM_HANDLER.setLevel(logging.DEBUG)
-ROOT_LOGGER.addHandler(LOG_STREAM_HANDLER)
+FILE_HANDLER = logging.FileHandler(SBLogger.make_log_path())
+FILE_HANDLER.setFormatter(logging.Formatter(fmt=LOGGING_FMT))
+ROOT_LOGGER.addHandler(FILE_HANDLER)
diff --git a/Engine/Plugins/VirtualProduction/Switchboard/Source/Switchboard/switchboard/util/collect_logs.py b/Engine/Plugins/VirtualProduction/Switchboard/Source/Switchboard/switchboard/util/collect_logs.py
--- a/Engine/Plugins/VirtualProduction/Switchboard/Source/Switchboard/switchboard/util/collect_logs.py
+++ b/Engine/Plugins/VirtualProduction/Switchboard/Source/Switchboard/switchboard/util/collect_logs.py
@@ -12,7 +12,7 @@ from PySide6 import QtCore, QtGui, QtWidgets
 from switchboard.config import Config
 from switchboard.devices.device_base import DeviceStatus
 from switchboard.devices.unreal.plugin_unreal import DeviceUnreal
-from ..switchboard_logging import LOGGER
+from ..switchboard_logging import LOGGER, FILE_HANDLER
 
 
 def open_logs_folder():
@@ -145,4 +145,4 @@ def iterate_log_files(config: Config, devices: typing.List[DeviceUnreal], consum
     iterate_mu_session_files()
     consumer(config.multiuser_server_log_path(), os.path.join("MultiUserServer", "UnrealMultiUserServer.log"))
 
-    consumer(LOGGER.save_log_file(), "Switchboard.log")
+    consumer(FILE_HANDLER.baseFilename, "Switchboard.log")
