Switchboard: Add Live Link Hub device.

The Live Link Hub device derives from the Unreal device, and surfaces a very stripped down subset of related functionality.

The choice to derive from DeviceUnreal is primarily motivated by the requirement to correctly participate in the sync/build/launch state machine, particularly in the case where another Unreal or nDisplay device may be sharing a workspace with the Live Link Hub device. It also gives us retrieval of Live Link Hub log files on close.

You can denote a single "hub" device using the same radio selection mechanism as nDisplay primary node selection, and all other devices are inferred to be operating in "spoke" mode.

Also removes a couple of noisy log lines from switchboard_widgets.AddressComboBox.

diff --git a/Engine/Plugins/VirtualProduction/Switchboard/Source/Switchboard/switchboard/device_list_widget.py b/Engine/Plugins/VirtualProduction/Switchboard/Source/Switchboard/switchboard/device_list_widget.py
--- a/Engine/Plugins/VirtualProduction/Switchboard/Source/Switchboard/switchboard/device_list_widget.py
+++ b/Engine/Plugins/VirtualProduction/Switchboard/Source/Switchboard/switchboard/device_list_widget.py
@@ -138,6 +138,9 @@ class DeviceListWidget(QtWidgets.QListWidget):
         header_widget = self.itemWidget(header_item)
         assert isinstance(header_widget, DeviceWidgetHeader)
 
+        if header_widget.autojoin_mu is None:
+            return
+
         any_set = False
         for device in self.devices():
             is_autojoin_set = device.autojoin_mu_server and device.autojoin_mu_server.get_value()
diff --git a/Engine/Plugins/VirtualProduction/Switchboard/Source/Switchboard/switchboard/devices/livelinkhub/plugin_livelinkhub.py b/Engine/Plugins/VirtualProduction/Switchboard/Source/Switchboard/switchboard/devices/livelinkhub/plugin_livelinkhub.py
new file mode 100644
--- /dev/null
+++ b/Engine/Plugins/VirtualProduction/Switchboard/Source/Switchboard/switchboard/devices/livelinkhub/plugin_livelinkhub.py
@@ -0,0 +1,293 @@
+# Copyright Epic Games, Inc. All Rights Reserved.
+
+from __future__ import annotations
+
+import pathlib
+from typing import Callable
+
+from PySide6 import QtCore, QtWidgets
+
+from switchboard import message_protocol, switchboard_widgets
+from switchboard.config import CONFIG, Setting, BoolSetting, StringSetting
+from switchboard.devices.device_base import Device, DeviceStatus, PluginHeaderWidgets
+from switchboard.devices.unreal.plugin_unreal import DeviceUnreal, DeviceWidgetUnreal, ProgramStartQueueItem
+from switchboard.switchboard_logging import LOGGER
+
+
+class DeviceLiveLinkHub(DeviceUnreal):
+    LLH_PROG_NAME = 'livelinkhub'
+
+    all_llh_devices: set[DeviceLiveLinkHub] = set()
+
+    csettings: dict[str, Setting] = {
+        'primary_device_name': StringSetting(
+            attr_name='primary_device_name',
+            nice_name='Primary Device',
+            value='',
+            tool_tip='Identifies which Live Link Hub device should operate in "hub" (non-spoke) mode',
+            show_ui=False,
+        ),
+        'executable_name': StringSetting(
+            attr_name='executable_name',
+            nice_name='Live Link Hub executable filename',
+            value='LiveLinkHub.exe',
+            category='General Settings',
+        ),
+        'command_line_arguments': StringSetting(
+            attr_name='command_line_arguments',
+            nice_name='Command Line Arguments',
+            value='',
+            tool_tip='Additional command line arguments for Live Link Hub',
+            category='General Settings',
+        ),
+        'retrieve_logs': BoolSetting(
+            attr_name='retrieve_logs',
+            nice_name='Retrieve Logs',
+            value=True,
+            tool_tip='When checked, retrieves logs after Live Link Hub terminates',
+            category='General Settings',
+        ),
+        'port': DeviceUnreal.csettings['port'],
+    }
+
+    def __init__(self, name, address, *args, **kwargs):
+        super().__init__(name, address, *args, **kwargs)
+
+        self.widget: DeviceWidgetLiveLinkHub
+
+    @property
+    def llh_proj_dir(self) -> pathlib.Path:
+        ''' Returns the directory containing LiveLinkHub.uproject (and logs, etc). '''
+        engine_dir = pathlib.Path(CONFIG.ENGINE_DIR.get_value(self.name))
+        return engine_dir / 'Source' / 'Programs' / 'LiveLinkHub'
+
+    @classmethod
+    def select_device_as_primary(cls, selected_device: DeviceLiveLinkHub):
+        ''' Selects the given device as the "hub" (non-spoke) device '''
+
+        cls.csettings['primary_device_name'].update_value(selected_device.name)
+
+        for dev in cls.all_llh_devices:
+            if dev.widget:
+                dev.widget.primary_button.setChecked(dev.name == selected_device.name)
+
+    def is_primary(self):
+        ''' Check if the device is the "hub" (non-spoke) device or not '''
+        return self.name == DeviceLiveLinkHub.csettings['primary_device_name'].get_value()
+
+    def select_as_primary(self):
+        ''' Selects this device as the "hub" (non-spoke) device '''
+        DeviceLiveLinkHub.select_device_as_primary(self)
+
+    #@override
+    @classmethod
+    def plugin_settings(cls):
+        return list(cls.csettings.values())
+
+    #@override
+    def device_settings(self):
+        return Device.device_settings(self) + [
+        ]
+
+    #@override
+    def setting_overrides(self):
+        return [
+            DeviceLiveLinkHub.csettings['command_line_arguments'],
+            CONFIG.ENGINE_DIR,
+            CONFIG.SOURCE_CONTROL_WORKSPACE,
+        ]
+
+    #@override
+    @classmethod
+    def plugin_header_widget_config(cls):
+        return super().plugin_header_widget_config() & ~PluginHeaderWidgets.AUTOJOIN_MU
+
+    #@override
+    @classmethod
+    def added_device(cls, device: DeviceLiveLinkHub):
+        super().added_device(device)
+
+        assert device not in DeviceLiveLinkHub.all_llh_devices
+        DeviceLiveLinkHub.all_llh_devices.add(device)
+
+        if not cls.csettings['primary_device_name'].get_value():
+            # Automatically make the first added device "primary"
+            cls.select_device_as_primary(device)
+
+    #@override
+    @classmethod
+    def removed_device(cls, device: DeviceLiveLinkHub):
+        super().removed_device(device)
+
+        assert device in DeviceLiveLinkHub.all_llh_devices
+        DeviceLiveLinkHub.all_llh_devices.remove(device)
+
+    #@override
+    def device_widget_registered(self, device_widget: DeviceWidgetLiveLinkHub):
+        super().device_widget_registered(device_widget)
+
+        device_widget.signal_device_widget_primary.connect(self.select_as_primary)
+        device_widget.primary_button.setChecked(self.is_primary())
+
+    #@override
+    @property
+    def executable_filename(self):
+        return self.csettings['executable_name'].get_value()
+
+    #@override
+    @property
+    def extra_cmdline_args_setting(self):
+        return DeviceLiveLinkHub.csettings['command_line_arguments'].get_value(self.name)
+
+    #@override
+    def get_remote_log_path(self):
+        return self.llh_proj_dir / 'Saved' / 'Logs'
+
+    #@override
+    def generate_unreal_command_line_args(self, map_name: str):
+        command_line_args = f'{self.extra_cmdline_args_setting}'
+
+        command_line_args += f' Log={self.log_filename}'
+
+        command_line_args += ' -hub' if self.is_primary() else ' -spoke'
+
+        # TODO: Specify config to load. Need to implement on LLH side first
+
+        return command_line_args
+
+    #@override
+    def launch(self, map_name: str):
+        map_name = ''
+
+        exe_path, args = self.generate_unreal_command_line(map_name)
+        LOGGER.info(f"Launching Live Link Hub: {exe_path} {args}")
+        self.last_launch_command.update_value(f'{exe_path} {args}')
+
+        puuid, msg = message_protocol.create_start_process_message(
+            prog_path=exe_path,
+            prog_args=args,
+            prog_name=self.LLH_PROG_NAME,
+            caller=self.name,
+            update_clients_with_stdout=False,
+        )
+
+        self.program_start_queue.add(
+            ProgramStartQueueItem(
+                name=self.LLH_PROG_NAME,
+                puuid_dependency=None,
+                puuid=puuid,
+                msg_to_unreal_client=msg,
+            ),
+            unreal_client=self.unreal_client,
+        )
+
+    #@override
+    def close(self, force=False):
+        llh_puuids = self.program_start_queue.running_puuids_named(self.LLH_PROG_NAME)
+
+        if not llh_puuids:
+            self.status = DeviceStatus.CLOSED
+        else:
+            self.status = DeviceStatus.CLOSING
+            for llh_puuid in llh_puuids:
+                _, msg = message_protocol.create_kill_process_message(llh_puuid)
+                self.unreal_client.send_message(msg)
+
+    #@override
+    def do_program_running_update(self, prog: ProgramStartQueueItem):
+        super().do_program_running_update(prog)
+
+        if prog.name == self.LLH_PROG_NAME:
+            self.status = DeviceStatus.OPEN
+
+    #@override
+    def do_program_ended_update(
+        self,
+        *,
+        program_name: str,
+        returncode: int,
+        get_stdout_str: Callable[[], str]
+    ):
+        super().do_program_ended_update(
+            program_name=program_name,
+            returncode=returncode,
+            get_stdout_str=get_stdout_str,
+        )
+
+        if program_name == self.LLH_PROG_NAME:
+            self.status = DeviceStatus.CLOSED
+            if DeviceLiveLinkHub.csettings['retrieve_logs'].get_value():
+                self.start_retrieve_log(returncode)
+
+    #@override
+    def _request_roles_file(self):
+        # Not relevant to our device type.
+        return
+
+
+class DeviceWidgetLiveLinkHub(DeviceWidgetUnreal):
+    signal_device_widget_primary = QtCore.Signal(object)
+
+    def __init__(self, *args, **kwargs):
+        super().__init__(*args, **kwargs)
+
+        self.primary_button: switchboard_widgets.ControlQPushButton
+
+    #@override
+    @property
+    def app_name(self):
+        return 'Live Link Hub'
+
+    #@override
+    def _add_control_buttons(self):
+        self._autojoin_visible = False
+        super()._add_control_buttons()
+
+    #@override
+    def add_widget_to_layout(self, widget):
+        if widget == self.name_line_edit:
+            self._add_primary_button()
+
+            # shorten the widget to account for the inserted primary button
+            btn_added_width = (
+                max(self.primary_button.iconSize().width(),
+                    self.primary_button.minimumSize().width())
+                + 2 * self.layout.spacing()
+                + self.primary_button.contentsMargins().left()
+                + self.primary_button.contentsMargins().right())
+
+            le_maxwidth = self.name_line_edit.maximumWidth() - btn_added_width
+            self.name_line_edit.setMaximumWidth(le_maxwidth)
+
+        super().add_widget_to_layout(widget)
+
+    #@override
+    def populate_context_menu(self, cmenu: QtWidgets.QMenu):
+        cmenu.addAction('Include in build' if self.exclude_from_build else 'Exclude from build',
+                        lambda: self.signal_exclude_from_build_toggled.emit())
+        cmenu.addAction('Open fetched log', lambda: self.signal_open_last_log.emit(self))
+        cmenu.addAction('Copy last launch command', lambda: self.signal_copy_last_launch_command.emit(self))
+
+    def _add_primary_button(self):
+        '''
+        Adds to the layout a button to select which device should be the
+        "hub"-mode device; all others are inferred to be "spoke" mode.
+        '''
+
+        self.primary_button = switchboard_widgets.ControlQPushButton.create(
+            ':/icons/images/star_yellow_off.png',
+            icon_disabled=':/icons/images/star_yellow_off.png',
+            icon_hover=':/icons/images/star_yellow_off.png',
+            icon_disabled_on=':/icons/images/star_yellow_off.png',
+            icon_on=':/icons/images/star_yellow.png',
+            icon_hover_on=':/icons/images/star_yellow.png',
+            icon_size=QtCore.QSize(16, 16),
+            tool_tip='Select as "hub" (non-spoke) device'
+        )
+
+        self.add_widget_to_layout(self.primary_button)
+
+        self.primary_button.clicked.connect(self._on_primary_button_clicked)
+
+    def _on_primary_button_clicked(self):
+        self.signal_device_widget_primary.emit(self)
diff --git a/Engine/Plugins/VirtualProduction/Switchboard/Source/Switchboard/switchboard/devices/unreal/plugin_unreal.py b/Engine/Plugins/VirtualProduction/Switchboard/Source/Switchboard/switchboard/devices/unreal/plugin_unreal.py
--- a/Engine/Plugins/VirtualProduction/Switchboard/Source/Switchboard/switchboard/devices/unreal/plugin_unreal.py
+++ b/Engine/Plugins/VirtualProduction/Switchboard/Source/Switchboard/switchboard/devices/unreal/plugin_unreal.py
@@ -705,16 +705,14 @@ class DeviceUnreal(Device):
             attr_name='retrieve_logs',
             nice_name='Retrieve Logs',
             value=True,
-            tool_tip=(
-                'When checked, retrieves the logs and traces after Unreal Engine terminates. \n'),
+            tool_tip='When checked, retrieves the logs and traces after Unreal Engine terminates.',
             category="UE Settings",
         ),
         'livelink_preset': LiveLinkPresetSetting(
             attr_name='livelink_preset',
             nice_name='LiveLink Preset',
             value='',
-            tool_tip=(
-                'Adds the selected LiveLink preset to the command line \n'),
+            tool_tip='Adds the selected LiveLink preset to the command line',
             category="Tools Settings",
         ),
         'mediaprofile': MediaProfileSetting(
@@ -2588,6 +2586,19 @@ class DeviceUnreal(Device):
                 f"{self.name}: on_program_ended with unknown id {puuid}")
             return
 
+        self.do_program_ended_update(
+            program_name=program_name,
+            returncode=returncode,
+            get_stdout_str=get_stdout_str,
+        )
+
+    def do_program_ended_update(
+        self,
+        *,
+        program_name: str,
+        returncode: int,
+        get_stdout_str: Callable[[], str]
+    ):
         # Update the widget classes (e.g. open button icon) once the
         # program_start_queue has been updated, since the status of the
         # queue may affect the widget classes.
@@ -3478,6 +3489,10 @@ class DeviceWidgetUnreal(DeviceWidget):
         CONFIG.P4_ENABLED.signal_setting_changed.connect(
             lambda _, enabled: self.sync_button.setVisible(enabled))
 
+    @property
+    def app_name(self):
+        return 'Unreal'
+
     @property
     def exclude_from_build(self):
         return self._exclude_from_build
@@ -3548,7 +3563,7 @@ class DeviceWidgetUnreal(DeviceWidget):
 
         self.open_button = self.add_control_button(
             icon_size=CONTROL_BUTTON_ICON_SIZE,
-            tool_tip='Start Unreal',
+            tool_tip=f'Start {self.app_name}',
             hover_focus=False,
             name='open')
 
@@ -3692,11 +3707,11 @@ class DeviceWidgetUnreal(DeviceWidget):
 
         if self.open_button.isChecked():
             if self.open_button.isEnabled():
-                self.open_button.setToolTip('Stop Unreal')
+                self.open_button.setToolTip(f'Stop {self.app_name}')
             else:
-                self.open_button.setToolTip('Stopping Unreal...')
+                self.open_button.setToolTip(f'Stopping {self.app_name}...')
         else:
-            self.open_button.setToolTip('Start Unreal')
+            self.open_button.setToolTip(f'Start {self.app_name}')
 
     def update_project_changelist(self, required_cl: str, current_device_cl: str):
         self.project_changelist_label.setText(f'P: {current_device_cl}')
diff --git a/Engine/Plugins/VirtualProduction/Switchboard/Source/Switchboard/switchboard/switchboard_widgets.py b/Engine/Plugins/VirtualProduction/Switchboard/Source/Switchboard/switchboard/switchboard_widgets.py
--- a/Engine/Plugins/VirtualProduction/Switchboard/Source/Switchboard/switchboard/switchboard_widgets.py
+++ b/Engine/Plugins/VirtualProduction/Switchboard/Source/Switchboard/switchboard/switchboard_widgets.py
@@ -469,7 +469,6 @@ class AddressComboBox(NonScrollableComboBox):
         self.runnable = AddressComboBox.ResolverRunnable(name)
         self.runnable.signals.result.connect(self.on_resolver_result)
         QtCore.QThreadPool.globalInstance().start(self.runnable)
-        LOGGER.debug(f'starting resolve: "{name}"')
         self.resolving = True
         return True
 
@@ -509,8 +508,6 @@ class AddressComboBox(NonScrollableComboBox):
             self.start_resolve()
             return
 
-        LOGGER.debug(f'resolved: "{", ".join(addrs)}"')
-
         if len(addrs) == 0:
             set_qt_property(self, 'validation', 'invalid')
         else:
