Switchboard: Don't use potentially overridden SWITCHBOARD_CONFIG_DIR for user_settings.json.

This should always be local / per-machine, even in cases where users have chosen a shared network location.

#rnx

diff --git a/Engine/Plugins/VirtualProduction/Switchboard/Source/Switchboard/switchboard/config.py b/Engine/Plugins/VirtualProduction/Switchboard/Source/Switchboard/switchboard/config.py
--- a/Engine/Plugins/VirtualProduction/Switchboard/Source/Switchboard/switchboard/config.py
+++ b/Engine/Plugins/VirtualProduction/Switchboard/Source/Switchboard/switchboard/config.py
@@ -83,12 +83,14 @@ class ConfigManager:
     @property
     def user_settings_file_path(self) -> pathlib.Path:
         '''Path to the "global" settings, including the last loaded config.'''
-        return self.config_dir.joinpath(self.USER_SETTINGS_FILE_NAME)
+        # Note: We use the default dir here to keep this local even if configs are shared.
+        return self.DEFAULT_CONFIG_DIR.joinpath(self.USER_SETTINGS_FILE_NAME)
 
     @property
     def user_settings_backup_file_path(self) -> pathlib.Path:
         '''Path at which a backup is automatically created if we're unable to load.'''
-        return self.config_dir.joinpath(self.USER_SETTINGS_BACKUP_FILE_NAME)
+        # Note: We use the default dir here to keep this local even if configs are shared.
+        return self.DEFAULT_CONFIG_DIR.joinpath(self.USER_SETTINGS_BACKUP_FILE_NAME)
 
     def list_config_paths(self) -> list[pathlib.Path]:
         '''
