Switchboard: Update Live Link Hub icon, pass `-UDPMESSAGING_SHARE_KNOWN_NODES=1` to "hub" instance.

#rnx

diff --git a/Engine/Plugins/VirtualProduction/Switchboard/Source/Switchboard/switchboard/devices/livelinkhub/plugin_livelinkhub.py b/Engine/Plugins/VirtualProduction/Switchboard/Source/Switchboard/switchboard/devices/livelinkhub/plugin_livelinkhub.py
--- a/Engine/Plugins/VirtualProduction/Switchboard/Source/Switchboard/switchboard/devices/livelinkhub/plugin_livelinkhub.py
+++ b/Engine/Plugins/VirtualProduction/Switchboard/Source/Switchboard/switchboard/devices/livelinkhub/plugin_livelinkhub.py
@@ -160,7 +160,11 @@ class DeviceLiveLinkHub(DeviceUnreal):
 
         command_line_args += f' Log={self.log_filename}'
 
-        command_line_args += ' -hub' if self.is_primary() else ' -spoke'
+        if self.is_primary():
+            command_line_args += ' -hub'
+            command_line_args += ' -UDPMESSAGING_SHARE_KNOWN_NODES=1'
+        else:
+            command_line_args += ' -spoke'
 
         if self.session_path:
             command_line_args += f' -SessionPath="{self.session_path}"'
