[Switchboard] Initial performance testing integration

#rnx

diff --git a/Engine/Plugins/VirtualProduction/ICVFXTesting/Build/Scripts/Automation/ICVFX.TestNode.cs b/Engine/Plugins/VirtualProduction/ICVFXTesting/Build/Scripts/Automation/ICVFX.TestNode.cs
--- a/Engine/Plugins/VirtualProduction/ICVFXTesting/Build/Scripts/Automation/ICVFX.TestNode.cs
+++ b/Engine/Plugins/VirtualProduction/ICVFXTesting/Build/Scripts/Automation/ICVFX.TestNode.cs
@@ -80,6 +80,42 @@ namespace ICVFXTest
 		[AutoParamWithNames(100.0f, "ICVFXTest.ViewportScreenPercentage")]
 		public float ViewportScreenPercentage;
 
+		/// <summary>
+		/// Should run fullscreen
+		/// </summary>
+		[AutoParamWithNames(false, "ICVFXTest.Fullscreen")]
+		public bool Fullscreen;
+
+		/// <summary>
+		/// Should run offscreen
+		/// </summary>
+		[AutoParamWithNames(false, "ICVFXTest.Offscreen")]
+		public bool Offscreen;
+
+		/// <summary>
+		/// Window X location
+		/// </summary>
+		[AutoParamWithNames(0, "ICVFXTest.WndLocX")]
+		public int WndLocX;
+
+		/// <summary>
+		/// Window Y location
+		/// </summary>
+		[AutoParamWithNames(0, "ICVFXTest.WndLocY")]
+		public int WndLocY;
+
+		/// <summary>
+		/// Window X size (width)
+		/// </summary>
+		[AutoParamWithNames(1920, "ICVFXTest.WndResX")]
+		public int WndResX;
+
+		/// <summary>
+		/// Window Y size (height)
+		/// </summary>
+		[AutoParamWithNames(1080, "ICVFXTest.WndResY")]
+		public int WndResY;
+
 		/// <summary>
 		/// Whether to use Lumen or not.
 		/// </summary>
diff --git a/Engine/Plugins/VirtualProduction/ICVFXTesting/Build/Scripts/Automation/Tests/ICVFXTest.AutoTest.cs b/Engine/Plugins/VirtualProduction/ICVFXTesting/Build/Scripts/Automation/Tests/ICVFXTest.AutoTest.cs
--- a/Engine/Plugins/VirtualProduction/ICVFXTesting/Build/Scripts/Automation/Tests/ICVFXTest.AutoTest.cs
+++ b/Engine/Plugins/VirtualProduction/ICVFXTesting/Build/Scripts/Automation/Tests/ICVFXTest.AutoTest.cs
@@ -5,6 +5,7 @@ using Gauntlet;
 using Microsoft.Extensions.Logging;
 using System.Collections.Generic;
 using System.Diagnostics;
+using System.IO;
 using UnrealBuildTool;
 
 using Log = EpicGames.Core.Log;
@@ -72,15 +73,22 @@ namespace ICVFXTest
 			Config.MaxDuration = Context.TestParams.ParseValue("MaxDuration", 60 * 60);  // 1 hour max
 			
 			UnrealTestRole ClientRole = Config.RequireRole(UnrealTargetRole.Client);
-            ClientRole.Controllers.Add("ICVFXTestControllerAutoTest");
+			ClientRole.Controllers.Add("ICVFXTestControllerAutoTest");
 
-			ClientRole.CommandLineParams.AddOrAppendParamValue("execcmds", "ICVFXTest.MaxRunCount " + Config.MaxRunCount);
+			ClientRole.CommandLineParams.AddOrAppendParamValue("DPCVars", "Slate.bAllowNotifications=0");
+			ClientRole.CommandLineParams.AddOrAppendParamValue("DPCVars", "p.Chaos.Solver.Deterministic=1");
 
+			ClientRole.CommandLineParams.AddOrAppendParamValue("execcmds", "ICVFXTest.MaxRunCount " + Config.MaxRunCount);
 			ClientRole.CommandLineParams.AddOrAppendParamValue("execcmds", "ICVFXTest.SoakTime " + Config.SoakTime);
 
 			if (Config.EnableTrace && Config.TraceRootFolder.Length != 0)
 			{
-				ClientRole.CommandLineParams.AddOrAppendParamValue("execcmds", "ICVFXTest.TraceFileName " + $"{Config.TraceRootFolder}/{GetTestSuffix()}.utrace");
+				string ClusterNode = Config.DisplayClusterNodeName;
+				string TraceFile = string.IsNullOrEmpty(ClusterNode) ?
+					$"{Config.TraceRootFolder}/{GetTestSuffix()}.utrace" :
+					$"{Config.TraceRootFolder}/{ClusterNode}_{GetTestSuffix()}.utrace";
+
+				ClientRole.CommandLineParams.AddOrAppendParamValue("execcmds", "ICVFXTest.TraceFileName " + $"{TraceFile}");
 			}
 
 			ClientRole.CommandLineParams.AddOrAppendParamValue("execcmds", "r.nanite " + (UseNanite() ? 1 : 0));
@@ -97,10 +105,10 @@ namespace ICVFXTest
 
 			// Temporary until we figure out why we hang
 			ClientRole.CommandLineParams.AddOrAppendParamValue("execcmds", "FX.AllowGPUParticles 0");
-
 			ClientRole.CommandLineParams.AddOrAppendParamValue("execcmds", "r.Shadow.Virtual.Enable " + (Config.DisableVirtualShadowMaps ? 0 : 1));
-			
-			ClientRole.CommandLineParams.Add("-MaxGPUCount=" + GetMaxGPUCount());
+			ClientRole.CommandLineParams.AddOrAppendParamValue("execcmds", "DisableAllScreenMessages");
+
+			ClientRole.CommandLineParams.Add("MaxGPUCount=" + GetMaxGPUCount());
 
 			if (Config.SkipTestSequence)
 			{
@@ -109,7 +117,8 @@ namespace ICVFXTest
 
 			if (Config.StatCommands)
 			{
-				ClientRole.CommandLineParams.AddOrAppendParamValue("execcmds", "stat fps, stat unitgraph");
+				ClientRole.CommandLineParams.AddOrAppendParamValue("execcmds", "stat fps");
+				ClientRole.CommandLineParams.AddOrAppendParamValue("execcmds", "stat unitgraph");
 			}
 			
 			if (Config.FPSChart)
@@ -135,13 +144,13 @@ namespace ICVFXTest
 
 			if (Config.D3DDebug)
 			{
-				ClientRole.CommandLineParams.Add("-d3ddebug");
+				ClientRole.CommandLineParams.Add("d3ddebug");
 				Gauntlet.Log.Info("Running with D3DDebug");
 			}
 
 			if (Config.GPUCrashDebugging)
 			{
-				ClientRole.CommandLineParams.Add("-gpucrashdebugging");
+				ClientRole.CommandLineParams.Add("gpucrashdebugging");
 				Gauntlet.Log.Info("Running with GPUCrashDebugging");
 			}
 
@@ -178,34 +187,69 @@ namespace ICVFXTest
 
 			if (UseVulkan())
 			{
-				ClientRole.CommandLineParams.Add("-vulkan");
+				ClientRole.CommandLineParams.Add("vulkan");
 			}
 			else
 			{
-				ClientRole.CommandLineParams.Add("-dx12");
+				ClientRole.CommandLineParams.Add("dx12");
 			}
 
 			if (Config.RDGImmediate)
 			{
-				ClientRole.CommandLineParams.Add("-rdgimmediate");
+				ClientRole.CommandLineParams.Add("rdgimmediate");
 			}
 			
 			if (Config.NoRHIThread)
 			{
-				ClientRole.CommandLineParams.Add("-norhithread");
+				ClientRole.CommandLineParams.Add("norhithread");
 			}
 
 			if (Config.EnableTrace)
 			{
 				// Start as late as possible to reduce trace file size.s
-				ClientRole.CommandLineParams.Add("-traceautostart=0");
-				ClientRole.CommandLineParams.Add("-trace=cpu,counters,gpu,bookmark");
+				ClientRole.CommandLineParams.Add("traceautostart=0");
+				ClientRole.CommandLineParams.Add("trace=cpu,counters,gpu,bookmark");
 			}
 
-
-			ClientRole.CommandLineParams.Add("-messaging -dc_cluster -nosplash -fixedseed -NoVerifyGC -noxrstereo -xrtrackingonly -RemoteControlIsHeadless -dc_dev_mono -ini:Engine:[/Script/Engine.Engine]:GameEngine=/Script/DisplayCluster.DisplayClusterGameEngine,[/Script/Engine.Engine]:GameViewportClientClassName=/Script/DisplayCluster.DisplayClusterViewportClient,[/Script/Engine.UserInterfaceSettings]:bAllowHighDPIInGameMode=True -ini:Game:[/Script/EngineSettings.GeneralProjectSettings]:bUseBorderlessWindow=True -unattended -NoScreenMessages -handleensurepercent=0 -UDPMESSAGING_TRANSPORT_MULTICAST=\"230.0.0.1:6666\" -UDPMESSAGING_TRANSPORT_UNICAST=\"127.0.0.1:0\" -UDPMESSAGING_TRANSPORT_STATIC=\"127.0.0.1:9030\" -ExecCmds=\"DisableAllScreenMessages\" -windowed -DPCVars=\"Slate.bAllowNotifications=0,p.Chaos.Solver.Deterministic=1\"");
-			// Todo: Is there a way to not specify this resolution? At the moment this matches what's in the display config file.
-			ClientRole.CommandLineParams.Add("-forceres WinX=0 WinY=0 -ResX=1920 -ResY=1080 -DisablePlugins=\"RenderDocPlugin\"");
+			ClientRole.CommandLineParams.Add("unattended");
+			ClientRole.CommandLineParams.Add("NoScreenMessages");
+			ClientRole.CommandLineParams.Add("handleensurepercent=0");
+			ClientRole.CommandLineParams.Add("messaging");
+			ClientRole.CommandLineParams.Add("dc_cluster");
+			ClientRole.CommandLineParams.Add("nosplash");
+			ClientRole.CommandLineParams.Add("NoLoadingScreen");
+			ClientRole.CommandLineParams.Add("fixedseed");
+			ClientRole.CommandLineParams.Add("NoVerifyGC");
+			ClientRole.CommandLineParams.Add("noxrstereo");
+			ClientRole.CommandLineParams.Add("xrtrackingonly");
+			ClientRole.CommandLineParams.Add("RemoteControlIsHeadless");
+			ClientRole.CommandLineParams.Add("UDPMESSAGING_TRANSPORT_MULTICAST=\"230.0.0.1:6666\"");
+			ClientRole.CommandLineParams.Add("UDPMESSAGING_TRANSPORT_UNICAST=\"127.0.0.1:0\"");
+			ClientRole.CommandLineParams.Add("UDPMESSAGING_TRANSPORT_STATIC=\"127.0.0.1:9030\"");
+			ClientRole.CommandLineParams.Add("ini:Engine:[/Script/Engine.Engine]:GameEngine=/Script/DisplayCluster.DisplayClusterGameEngine,[/Script/Engine.Engine]:GameViewportClientClassName=/Script/DisplayCluster.DisplayClusterViewportClient,[/Script/Engine.UserInterfaceSettings]:bAllowHighDPIInGameMode=True");
+			ClientRole.CommandLineParams.Add("ini:Game:[/Script/EngineSettings.GeneralProjectSettings]:bUseBorderlessWindow=True");
+			ClientRole.CommandLineParams.Add("DisablePlugins=\"RenderDocPlugin\"");
+			ClientRole.CommandLineParams.Add("dc_dev_mono");
+
+			if (Config.Fullscreen)
+			{
+				ClientRole.CommandLineParams.Add("fullscreen");
+			}
+			else
+			{
+				ClientRole.CommandLineParams.Add("forceres");
+				ClientRole.CommandLineParams.Add("windowed");
+				ClientRole.CommandLineParams.Add($"WinX={Config.WndLocX}");
+				ClientRole.CommandLineParams.Add($"WinY={Config.WndLocY}");
+				ClientRole.CommandLineParams.Add($"ResX={Config.WndResX}");
+				ClientRole.CommandLineParams.Add($"ResY={Config.WndResY}");
+				ClientRole.CommandLineParams.AddOrAppendParamValue("execcmds", $"r.setres {Config.WndResX}x{Config.WndResY}");
+
+				if (Config.Offscreen)
+				{
+					ClientRole.CommandLineParams.Add("RenderOffscreen");
+				}
+			}
 
 			Gauntlet.Log.Info(ClientRole.CommandLineParams.ToString());
 
@@ -213,7 +257,7 @@ namespace ICVFXTest
 		}
 
 		protected override void InitHandledErrors()
-        {
+		{
 			base.InitHandledErrors();
 			HandledErrors.Add(new HandledError("AutoTest", "AutoTest failure:", "LogICVFXTest", true));
 		}
diff --git a/Engine/Plugins/VirtualProduction/ICVFXTesting/Build/Scripts/Automation/Tests/ICVFXTest.PerformanceReport.cs b/Engine/Plugins/VirtualProduction/ICVFXTesting/Build/Scripts/Automation/Tests/ICVFXTest.PerformanceReport.cs
--- a/Engine/Plugins/VirtualProduction/ICVFXTesting/Build/Scripts/Automation/Tests/ICVFXTest.PerformanceReport.cs
+++ b/Engine/Plugins/VirtualProduction/ICVFXTesting/Build/Scripts/Automation/Tests/ICVFXTest.PerformanceReport.cs
@@ -89,7 +89,7 @@ namespace ICVFXTest
 	/// </summary>
 	public class PerformanceReport : AutoTest
 	{
-		private DirectoryInfo TempPerfCSVDir => new DirectoryInfo(Path.Combine(Unreal.RootDirectory.FullName, "GauntletTemp", "PerfReportCSVs"));
+		private DirectoryInfo TempPerfCSVDir => new DirectoryInfo(Path.Combine(Unreal.RootDirectory.FullName, "GauntletTemp", "PerfReportCSVs_" + GetConfiguration().DisplayClusterNodeName));
 		private static ILogger Logger => Log.Logger;
 		private string OriginalBuildName;
 		private string OverrideDisplayConfigPath;
diff --git a/Engine/Plugins/VirtualProduction/ICVFXTesting/Scripts/GauntletSettings.xml b/Engine/Plugins/VirtualProduction/ICVFXTesting/Scripts/GauntletSettings.xml
--- a/Engine/Plugins/VirtualProduction/ICVFXTesting/Scripts/GauntletSettings.xml
+++ b/Engine/Plugins/VirtualProduction/ICVFXTesting/Scripts/GauntletSettings.xml
@@ -13,6 +13,19 @@
 	<Option Name="WithTargetAutomationMapsCycle" Restrict="true|false" DefaultValue="false" Description="Whether to run UE.TargetAutomation(RunTest=Project.Maps.Cycle)." />
 	<Option Name="WithTargetAutomationAll" Restrict="true|false" DefaultValue="false" Description="Whether to run UE.TargetAutomation(RunTest=Project.)." />
 
+	<!-- For on-set testing, there will be custom approach -->
+	<Option Name="TestNameExt" DefaultValue="" Description="Extended test name (sub-test name for on-stage runs)" />
+	<Option Name="Fullscreen" Restrict="true|false" DefaultValue="false" Description="Should run fullscreen" />
+	<Option Name="Offscreen" Restrict="true|false" DefaultValue="false" Description="Should run offscreen" />
+	<Option Name="WndLocX" DefaultValue="0" Description="Window X position" />
+	<Option Name="WndLocY" DefaultValue="0" Description="Window Y position" />
+	<Option Name="WndResX" DefaultValue="1920" Description="Window X size" />
+	<Option Name="WndResY" DefaultValue="1080" Description="Window Y size" />
+	<Option Name="IsOnStagePerfTest" Restrict="true|false" DefaultValue="false" Description="Whether it's running on a real stage" />
+	<Do If="$(IsOnStagePerfTest)">
+		<Property Name="EditorTestList" Value="ICVFXTest.PerformanceReport+" If="$(WithPerformanceReportTest)" />
+	</Do>
+
 	<!-- Hack solution to set target test list for BuildAndTest.xml -->
 	<Do If="'$(EditorTestList)' == ''">
 		<Property Name="TargetTestList" Value="$(EditorTestList)ICVFXTest.MemoryReport+" If="$(WithMemoryReportTest)" />
@@ -34,7 +47,7 @@
 	<!-- Test specific options -->
 	<Option Name="DisplayConfigPath" DefaultValue="" Description="The path to the .ndisplay file used for launching the cluster."/>
 	<Option Name="DisplayClusterNodeName" DefaultValue="" Description="The name of the ndisplay node to launch."/>
-	<Option Name="ViewportScreenPercentage" DefaultValue="100.0" Description="The viewport screen percentage to use"/>
+	<Option Name="ViewportScreenPercentage" DefaultValue="100" Description="The viewport screen percentage to use"/>
 	<Option Name="WithNanite" Restrict="true|false" DefaultValue="false" Description="Whether to use nanite when testing"/>
 	<Option Name="WithLumen" Restrict="true|false" DefaultValue="false" Description="Whether to use lumen when testing"/>
 	<Option Name="WithD3DDebug" Restrict="true|false" DefaultValue="false" Description="Whether to add d3ddebug command line argument."/>
@@ -58,9 +71,10 @@
 	<Property Name="AutomationNetworkRootDir" Value="\\epicgames.net\root\Builds\Automation\Reports\$(ProjectName)Perf"/>
 	<Property Name="AutomationNetworkRootDir" Value="/Volumes/Builds/Automation/Reports/$(ProjectName)Perf" If="'$(HostPlatform)' != 'Win64'"/>
 	<Property Name="AutomationLocalOutputDir" Value="$(RootDir)/LocalBuilds/Reports/$(ProjectName)Perf"/>
+	<Property Name="AutomationLocalOutputDir" Value="$(RootDir)/LocalBuilds/Reports/$(ProjectName)/$(TestNameExt)/$(DisplayClusterNodeName)" If="$(IsOnStagePerfTest)" />
 
 	<!-- Outputs for this run keeping slashes flipped the right way because it shows up on the webapp this way -->
-	<Property Name="AutomationNetworkOutputDir" Value="$(AutomationNetworkRootDir)\$(BuildNamePath)"/>
+	<Property Name="AutomationNetworkOutputDir" Value="$(AutomationNetworkRootDir)/$(BuildNamePath)"/>
 	<Property Name="AutomationNetworkOutputDir" Value="$(AutomationNetworkRootDir)/$(BuildNamePath)" If="'$(HostPlatform)' != 'Win64'"/>
 	<Property Name="AutomationOutputDir" Value="$(AutomationLocalOutputDir)"/>
 	<Property Name="AutomationOutputDir" Value="$(AutomationNetworkOutputDir)" If="$(IsBuildMachine)"/>
@@ -73,6 +87,8 @@
 	<Property Name="TestOutputDir" Value="$(AutomationOutputDir)/$(ProjectName)$(Platform)$(ProjectName)Perf"/>
 	<Property Name="LogOutputParams" Value="-logdir=&quot;$(TestOutputDir)&quot;"/>
 	
+	<Property Name="TraceRootFolder" Value="$(TestOutputDir)/$(BuildNamePath)" />
+	<Property Name="TraceRootFolder" Value="$(TestOutputDir)" If="$(IsOnStagePerfTest)" />	
 		
 	<Property Name="ExtraGauntletCommonArgs" Value="-screenshotperiod=0 -ICVFXTest.StatCommands -ICVFXTest.MaxRunCount=$(MaxRunCount) $(LogOutputParams)" />
 	<Property Name="ExtraGauntletEditorTestArgs" Value="-game" />
@@ -80,6 +96,8 @@
 	<Property Name="ExtraGauntletEditorTestArgs" Value="$(ExtraGauntletEditorTestArgs) -ICVFXTest.IsBuildMachine=$(IsBuildMachine)" />
 	<Property Name="ExtraGauntletEditorTestArgs" Value="$(ExtraGauntletEditorTestArgs) -ICVFXTest.TestName=$(ProjectTestName)" />
 	<Property Name="ExtraGauntletEditorTestArgs" Value="$(ExtraGauntletEditorTestArgs) -ICVFXTest.MapOverride=$(MapOverride) -ICVFXTest.DisplayConfig=$(DisplayConfigPath) -ICVFXTest.ViewportScreenPercentage=$(ViewportScreenPercentage)" />
+	<Property Name="ExtraGauntletEditorTestArgs" Value="$(ExtraGauntletEditorTestArgs) -ICVFXTest.Fullscreen=$(Fullscreen) -ICVFXTest.Offscreen=$(Offscreen)" />
+	<Property Name="ExtraGauntletEditorTestArgs" Value="$(ExtraGauntletEditorTestArgs) -ICVFXTest.WndLocX=$(WndLocX) -ICVFXTest.WndLocY=$(WndLocY) -ICVFXTest.WndResX=$(WndResX) -ICVFXTest.WndResY=$(WndResY)" />
 	<Property Name="ExtraGauntletEditorTestArgs" Value="$(ExtraGauntletEditorTestArgs) -ICVFXTest.DisplayClusterNodeName=$(DisplayClusterNodeName)" />
 	<Property Name="ExtraGauntletEditorTestArgs" Value="$(ExtraGauntletEditorTestArgs) -ICVFXTest.D3DDebug" If="$(WithD3DDebug)"/>
 	<Property Name="ExtraGauntletEditorTestArgs" Value="$(ExtraGauntletEditorTestArgs) -ICVFXTest.GPUCrashDebugging" If="$(WithGPUCrashDebugging)"/>
@@ -91,7 +109,7 @@
 	<Property Name="ExtraGauntletEditorTestArgs" Value="$(ExtraGauntletEditorTestArgs) -ICVFXTest.NoRHIThread" If="$(WithNoRHIThread)"/>
 	<Property Name="ExtraGauntletEditorTestArgs" Value="$(ExtraGauntletEditorTestArgs) -ICVFXTest.DisableVirtualShadowMaps" If="!$(WithVirtualShadowMaps)"/>
 	<Property Name="ExtraGauntletEditorTestArgs" Value="$(ExtraGauntletEditorTestArgs) -ICVFXTest.EnableTrace" If="$(WithTrace)"/>
-	<Property Name="ExtraGauntletEditorTestArgs" Value="$(ExtraGauntletEditorTestArgs) -ICVFXTest.TraceRootFolder=$(TestOutputDir)/$(BuildNamePath)" If="$(WithTrace)"/>
+	<Property Name="ExtraGauntletEditorTestArgs" Value="$(ExtraGauntletEditorTestArgs) -ICVFXTest.TraceRootFolder=$(TraceRootFolder)" If="$(WithTrace)"/>
 	<Property Name="ExtraGauntletEditorTestArgs" Value="$(ExtraGauntletEditorTestArgs) -ICVFXTest.PerfCacheFolder=$(ReportCacheDirRoot)"/>
 	<Property Name="ExtraGauntletEditorTestArgs" Value="$(ExtraGauntletEditorTestArgs) -ICVFXTest.NoLocalReports" If="$(NoLocalReports)"/>
 	<Property Name="ExtraGauntletEditorTestArgs" Value="$(ExtraGauntletEditorTestArgs) -ICVFXTest.SkipPerfReportServer" If="$(SkipPerfReportServer)"/>
diff --git a/Engine/Plugins/VirtualProduction/ICVFXTesting/Source/ICVFXTesting/Private/ICVFXTestControllerAutoTest.cpp b/Engine/Plugins/VirtualProduction/ICVFXTesting/Source/ICVFXTesting/Private/ICVFXTestControllerAutoTest.cpp
--- a/Engine/Plugins/VirtualProduction/ICVFXTesting/Source/ICVFXTesting/Private/ICVFXTestControllerAutoTest.cpp
+++ b/Engine/Plugins/VirtualProduction/ICVFXTesting/Source/ICVFXTesting/Private/ICVFXTestControllerAutoTest.cpp
@@ -18,6 +18,8 @@
 #include "PlatformFeatures.h"
 #include "ProfilingDebugging/CsvProfiler.h"
 #include "VideoRecordingSystem.h"
+#include "HAL/PlatformFilemanager.h"
+#include "Misc/Paths.h"
 #include "Roles/LiveLinkTransformRole.h"
 #include "UObject/SoftObjectPtr.h"
 
@@ -176,9 +178,18 @@ public:
 	{
 		Super::Start(PrevState);
 
-		if (!ICVFXTest::CVarTraceFileName.GetValueOnAnyThread().IsEmpty())
+		const FString TraceFile = ICVFXTest::CVarTraceFileName.GetValueOnAnyThread();
+		if (!TraceFile.IsEmpty())
 		{
-			Controller->ConsoleCommand(*(FString(TEXT("trace.file ")) + ICVFXTest::CVarTraceFileName.GetValueOnAnyThread()));
+			// Delete trace file if already exists, otherwise new traces won't be stored
+			IPlatformFile& PlatformFile = FPlatformFileManager::Get().GetPlatformFile();
+			if (PlatformFile.FileExists(*TraceFile))
+			{
+				UE_LOG(LogICVFXTest, Log, TEXT("Trace file '%s' already exists. Removing it..."), *TraceFile);
+				PlatformFile.DeleteFile(*TraceFile);
+			}
+
+			Controller->ConsoleCommand(*(FString(TEXT("trace.file ")) + TraceFile));
 		}
 
 		Controller->GoToTestLocation(0);
diff --git a/Engine/Plugins/VirtualProduction/Switchboard/Source/Switchboard/switchboard/devices/ndisplay/plugin_ndisplay.py b/Engine/Plugins/VirtualProduction/Switchboard/Source/Switchboard/switchboard/devices/ndisplay/plugin_ndisplay.py
--- a/Engine/Plugins/VirtualProduction/Switchboard/Source/Switchboard/switchboard/devices/ndisplay/plugin_ndisplay.py
+++ b/Engine/Plugins/VirtualProduction/Switchboard/Source/Switchboard/switchboard/devices/ndisplay/plugin_ndisplay.py
@@ -140,7 +140,7 @@ class AddnDisplayDialog(AddDeviceDialog):
         self.existing_ndisplay_devices = []
 
         for device in existing_devices:
-            if device.device_type == "nDisplay":
+            if device.device_type == "nDisplay" or str(device.device_type).startswith("nDisplay"):
                 self.existing_ndisplay_devices.append(device)
 
         if self.existing_ndisplay_devices:
@@ -802,10 +802,10 @@ class DevicenDisplay(DeviceUnreal):
     def create_monitor_if_necessary(cls):
         ''' Creates the nDisplay Monitor if it doesn't exist yet.
         '''
-        if not cls.ndisplay_monitor:
-            cls.ndisplay_monitor = nDisplayMonitor(None)
+        if not DevicenDisplay.ndisplay_monitor:
+            DevicenDisplay.ndisplay_monitor = nDisplayMonitor(None)
 
-        return cls.ndisplay_monitor
+        return DevicenDisplay.ndisplay_monitor
 
     def on_change_setting_affecting_command_line(self, oldval, newval):
         self.generate_unreal_command_line()
@@ -1556,7 +1556,7 @@ class DevicenDisplay(DeviceUnreal):
 
         devices = [
             dev for dev in cls.active_unreal_devices
-            if dev.device_type == 'nDisplay']
+            if dev.device_type == 'nDisplay' or str(dev.device_type).startswith('nDisplay')]
 
         shown_unknown_node_warning = False
 
diff --git a/Engine/Plugins/VirtualProduction/Switchboard/Source/Switchboard/switchboard/devices/ndisplayTest/ndisplayTest.py b/Engine/Plugins/VirtualProduction/Switchboard/Source/Switchboard/switchboard/devices/ndisplayTest/ndisplayTest.py
new file mode 100644
--- /dev/null
+++ b/Engine/Plugins/VirtualProduction/Switchboard/Source/Switchboard/switchboard/devices/ndisplayTest/ndisplayTest.py
@@ -0,0 +1,130 @@
+# Copyright Epic Games, Inc. All Rights Reserved.
+
+import copy
+import os
+import re
+import xml.etree.ElementTree as ET
+
+from pathlib import Path
+from switchboard.config import CONFIG
+from switchboard.switchboard_logging import LOGGER
+
+
+####################################################
+# A wrapper class for XML based nDisplay tests
+####################################################
+class nDisplayTest:
+
+    def __init__(self, xml_path):
+        self._file_path = xml_path
+        self._tree = None
+        self._root = None
+        self._ns = {"bg": "http://www.epicgames.com/BuildGraph"}
+        self._properties = { "RootDir": str(Path(CONFIG.ENGINE_DIR.get_value()).parent) }
+        self._load_xml()
+
+    def _load_xml(self):
+        ''' Attempts to load and parse the an XML file '''
+
+        try:
+            self._tree = ET.parse(self._file_path)
+            self._root = self._tree.getroot()
+            # Extract all properties into a dictionary
+            self._properties |= {
+                prop.attrib["Name"]: prop.attrib["Value"]
+                for prop in self._root.findall(".//bg:Property", self._ns)}
+        except:
+            self._tree = None
+            self._root = None
+            LOGGER.error(f"Could not read test file: '{self._file_path}'")
+
+    @classmethod
+    def parse_test(cls, file_path) -> "nDisplayTest":
+        ''' Parses an nDisplay test file into a test wrapper '''
+
+        try:
+            return cls(file_path) if os.path.exists(file_path) else None
+        except Exception:
+            LOGGER.error(f"Couldn't parse nDisplay test file: {file_path}")
+            return None
+
+    def duplicate(self) -> "nDisplayTest":
+        ''' Duplicates this instance, returns the new one '''
+
+        root_clone = copy.deepcopy(self._root)
+        cls = self.__class__
+        instance = cls.__new__(cls)  # Bypass __init__
+        instance._tree = ET.ElementTree(root_clone)
+        instance._root = root_clone
+        instance._file_path = self._file_path
+        instance._ns = self._ns
+        instance._properties = self._properties
+        return instance
+
+    def update_properties(self, updates: dict[str, str]):
+        ''' Updates properties with the new values (in memory) '''
+
+        # Build property map
+        prop_map = {
+            prop.attrib.get("Name"): prop
+            for prop in self._root.findall('.//bg:Property', self._ns)
+        }
+
+        # Iterate & update
+        for name, new_val in updates.items():
+            prop = prop_map.get(name)
+            if prop is not None:
+                prop.set("Value", new_val)
+
+    def to_string(self):
+        ''' Returns current XML data as string '''
+        ET.register_namespace('', "http://www.epicgames.com/BuildGraph")
+        return ET.tostring(self._root, encoding='utf-8', xml_declaration=True).decode('utf-8') if self.is_valid() else None
+
+    def to_bytes(self):
+        ''' Returns current XML data as bytes '''
+        ET.register_namespace('', "http://www.epicgames.com/BuildGraph")
+        return ET.tostring(self._root, encoding='utf-8', xml_declaration=True) if self.is_valid() else None
+
+    def save(self, path):
+        ''' Saves current XML data to file '''
+        if self.is_valid():
+            self._tree.write(path, encoding="utf-8", xml_declaration=True)
+        
+    def is_valid(self):
+        ''' Checks if XML was successfully loaded, and contains valid data '''
+        return self._tree is not None and self._root is not None
+
+    def _resolve_path(self, source):
+        ''' Resolves placeholders in paths from XML '''
+        output = re.sub(r"\$\((.*?)\)", lambda m: self._properties.get(m.group(1), m.group(0)), source)
+        return os.path.normpath(output)
+
+    def _get_property(self, property_name):
+        ''' Returns value of a specified XML property '''
+        return self._properties.get(property_name, "") if self.is_valid() else None
+
+    def get_test_path(self):
+        ''' Returns original test file path '''
+        return self._file_path
+
+    def get_ndisplay_config_path(self):
+        ''' Returns nDisplay config paths specified in an XML test '''
+
+        # Search for the config path property
+        config_path = self._get_property("DisplayConfigPath")
+
+        # Resolve placeholders if there are any
+        if config_path:
+            config_path = self._resolve_path(config_path)
+
+        return config_path
+
+    def get_test_map(self):
+        ''' Returns test map if specified '''
+        return self._get_property("MapOverride")
+
+    def get_script_dir(self):
+        ''' Returns script dir if specified '''
+        return self._get_property("ICVFXScriptDir")
+
diff --git a/Engine/Plugins/VirtualProduction/Switchboard/Source/Switchboard/switchboard/devices/ndisplayTest/plugin_ndisplayTest.py b/Engine/Plugins/VirtualProduction/Switchboard/Source/Switchboard/switchboard/devices/ndisplayTest/plugin_ndisplayTest.py
new file mode 100644
--- /dev/null
+++ b/Engine/Plugins/VirtualProduction/Switchboard/Source/Switchboard/switchboard/devices/ndisplayTest/plugin_ndisplayTest.py
@@ -0,0 +1,647 @@
+# Copyright Epic Games, Inc. All Rights Reserved.
+
+import glob
+import os
+import uuid
+
+from datetime import datetime
+from pathlib import Path
+from PySide6 import QtCore, QtWidgets
+from typing import Callable
+
+from switchboard import message_protocol
+from switchboard import switchboard_utils as sb_utils
+from switchboard import switchboard_widgets as sb_widgets
+from switchboard.config import CONFIG, StringSetting
+from switchboard.devices.device_base import DeviceStatus
+from switchboard.devices.device_widget_base import AddDeviceDialog
+from switchboard.devices.ndisplay.plugin_ndisplay import DevicenDisplay, DeviceWidgetnDisplay, LaunchMode
+from switchboard.devices.unreal.plugin_unreal import ProgramStartQueueItem, UnrealJobs
+from switchboard.switchboard_logging import LOGGER
+
+from .ndisplayTest import nDisplayTest
+
+
+####################################################
+# Add nDisplay Test device dialog
+####################################################
+class AddnDisplayTestDialog(AddDeviceDialog):
+    
+    def __init__(self, existing_devices, parent=None):
+        super().__init__(
+            device_type="nDisplayTest",
+            existing_devices=existing_devices,
+            parent=parent)
+
+        # Initialize members
+        self.existing_devices = existing_devices
+        self.test_dir = self.get_default_test_dir()
+        LOGGER.info(f"Test location: '{self.test_dir}'")
+
+        # Initialize GUI
+        self.initialize_widgets()
+
+    def initialize_widgets(self):
+        ''' Performs widgets initialization '''
+
+        # Set enough width for a decent file path length
+        self.setMinimumWidth(640)
+
+        # Remove unneeded base dialog layout rows
+        self.form_layout.removeRow(self.name_field)
+        self.form_layout.removeRow(self.address_field)
+        self.name_field = None
+        self.address_field = None
+
+        # Initialize grid layout
+        grid_layout = QtWidgets.QGridLayout()
+
+        # Test Dir: Label
+        lblTestDir = QtWidgets.QLabel(self, text="Tests Location")
+
+        # Test Dir: Text field
+        self.txtTestDir = QtWidgets.QLineEdit()
+        self.txtTestDir.setReadOnly(True)
+        self.txtTestDir.setText(self.test_dir)
+
+        # Test Dir: Button to browse for a test directory
+        self.btnBrowseTestDir = QtWidgets.QPushButton(self, text="Browse")
+        self.btnBrowseTestDir.clicked.connect(self.on_clicked_btnBrowseTestDir)
+
+        # Test Dir: Row widget
+        test_dir_layout = QtWidgets.QHBoxLayout()
+        test_dir_layout.addWidget(self.txtTestDir)
+        test_dir_layout.addWidget(self.btnBrowseTestDir)
+
+        # Test File: Label
+        lblTestFile = QtWidgets.QLabel(self, text="Test File")
+
+        # Test File: Combobox with the files
+        self.cbTests = sb_widgets.SearchableComboBox(self)
+        self.cbTests.setSizePolicy(QtWidgets.QSizePolicy.Expanding, QtWidgets.QSizePolicy.Preferred)
+        self.cbTests.setEditable(False)
+
+        # Test File: Row widget
+        test_file_layout = QtWidgets.QHBoxLayout()
+        test_file_layout.addWidget(self.cbTests)
+
+        # Initialize grid layout
+        grid_layout.addWidget(lblTestDir, 0, 0)
+        grid_layout.addLayout(test_dir_layout, 0, 1)
+        grid_layout.addWidget(lblTestFile, 1, 0)
+        grid_layout.addLayout(test_file_layout, 1, 1)
+
+        # Add the grid layout to the form layout
+        self.form_layout.addRow(grid_layout)
+
+        # Add a spacer right before the ok/cancel buttons
+        spacer_layout = QtWidgets.QHBoxLayout()
+        spacer_layout.addItem(
+            QtWidgets.QSpacerItem(20, 20, QtWidgets.QSizePolicy.Minimum, QtWidgets.QSizePolicy.Expanding))
+        self.form_layout.addRow("", spacer_layout)
+
+        # Add warning if there are any devices exist
+        if self.existing_devices:
+            self.layout().addWidget(
+                QtWidgets.QLabel("Warning! All existing nDisplay devices will be replaced."))
+
+        # Update GUI contents
+        self.refresh_test_list()
+
+    def set_test_location(self, test_location):
+        ''' Sets new test location and re-scans for available tests '''
+        LOGGER.info(f"Set new test location: {test_location}")
+        self.test_dir = test_location
+        self.txtTestDir.setText(self.test_dir)
+        self.refresh_test_list()
+
+    def refresh_test_list(self):
+        ''' Parses the currently selected test directory for all available test files (xml)'''
+
+        # Clear the current test list
+        self.cbTests.clear()
+
+        # Find all .xml files, and extract names only
+        test_files = glob.glob(f"{self.test_dir}/*.xml")
+        test_file_names = [os.path.splitext(os.path.basename(file))[0] for file in test_files]
+
+        # Update the combo data
+        self.cbTests.addItems(test_file_names)
+
+    def get_default_test_dir(self):
+        ''' Returns default test location '''
+
+        # Project root & potential test locations
+        project_dir = CONFIG.get_project_dir()
+        build_dir = os.path.join(project_dir, "Build")
+        automation_dir = os.path.join(build_dir, "Automation")
+
+        # Search for the existing test location
+        result_dir = next((p for p in [automation_dir, build_dir] if os.path.exists(p)), project_dir)
+
+        return result_dir
+
+    def on_clicked_btnBrowseTestDir(self):
+        ''' Opens a folder dialog to browse for the test directory '''
+
+        # Chose new location
+        dir_old = self.test_dir
+        dir_new = QtWidgets.QFileDialog.getExistingDirectory(self, "Choose tests location", self.test_dir)
+
+        # Update & refresh if new location is chosen
+        if os.path.exists(dir_new):
+            dir_new = os.path.normpath(dir_new)
+            if dir_old != dir_new:
+                self.set_test_location(dir_new)
+
+    def get_selected_test_file(self):
+        ''' Returns the test chosen '''
+        test_name = self.cbTests.currentText()
+        return os.path.join(self.test_dir, test_name + ".xml")
+
+    def devices_to_add(self):
+        ''' [override] Provides a list of devices to add '''
+
+        # Get selected test file path
+        test_path = self.get_selected_test_file()
+        if not test_path:
+            LOGGER.error(f"Wrong test path: {test_path}")
+            return []
+
+        # Parse the test file
+        test_chosen = nDisplayTest.parse_test(test_path)
+        if not test_chosen:
+            LOGGER.error(f"Couldn't parse test file: {test_path}")
+            return []
+
+        # Extract nDisplay config from the test file
+        config_path = test_chosen.get_ndisplay_config_path()
+        if not config_path:
+            LOGGER.error(f"No nDisplay configuration file found in test: {test_path}")
+            return []
+
+        # Pass data to the device class
+        DevicenDisplay.csettings['ndisplay_config_file'].update_value(config_path)
+        DevicenDisplay.csettings['launch_mode'].update_value(LaunchMode.Standalone.value)
+        DevicenDisplayTest.tsettings['ndisplay_test_file'].update_value(test_path)
+        DevicenDisplayTest.current_test = test_chosen
+
+        # Parse nDisplay config in order to get the list of devices
+        try:
+            devices = DevicenDisplay.parse_config(os.path.normpath(config_path)).nodes
+            if len(devices) == 0:
+                LOGGER.error(f"Could not read any devices in nDisplay config: {config_path}")
+
+            # Pass P-node ID
+            for node in devices:
+                if node['primary']:
+                    DevicenDisplay.csettings['primary_device_name'].update_value(node['name'])
+                    break
+                
+            return devices
+
+        except (IndexError, KeyError, FileNotFoundError):
+            LOGGER.error(f"Error parsing nDisplay config: {config_path}")
+            return []
+
+    def devices_to_remove(self):
+        ''' [override] Provides a list of devices to remove '''
+        # Remove all existing devices
+        return self.existing_devices
+
+
+####################################################
+# nDisplay Test device widget
+####################################################
+class DeviceWidgetnDisplayTest(DeviceWidgetnDisplay):
+
+    signal_open_test_artifacts = QtCore.Signal(object)
+
+    def populate_context_menu(self, cmenu: QtWidgets.QMenu):
+        ''' [override] Called to populate the given context menu with any desired actions '''
+        cmenu.addAction("Open fetched log", lambda: self.signal_open_last_log.emit(self))
+        cmenu.addAction("Open fetched trace", lambda: self.signal_open_last_trace.emit(self))
+        cmenu.addAction("Show artifacts", lambda: self.signal_open_test_artifacts.emit(self))
+        cmenu.addAction("Copy last launch command", lambda: self.signal_copy_last_launch_command.emit(self))
+
+
+####################################################
+# nDisplay Test device
+####################################################
+class DevicenDisplayTest(DevicenDisplay):
+
+    add_device_dialog = AddnDisplayTestDialog
+    current_test: nDisplayTest = None
+    artifacts_dir = None
+    
+    tsettings = {
+        'ndisplay_test_file': StringSetting(
+            attr_name="ndisplay_test_file",
+            nice_name="nDisplay Test File",
+            value="",
+            tool_tip="Path to nDisplay test file",
+            allow_reset=False,
+            is_read_only=True,
+            category="General Settings",
+        ),
+    }
+
+
+    def __init__(self, name, address, **kwargs):
+        super().__init__(name, address, **kwargs)
+
+        # "Send file" is an intermediate step on the way to launch a test
+        self.unreal_client.delegates['send file complete'] = self.on_send_file_complete
+
+        # Clean up temporaries
+        self.reset_launch_temporary_data()
+
+        # Load&parse test data if not done yet
+        if(self.current_test is None):
+            test_file = self.tsettings['ndisplay_test_file'].get_value()
+            self.current_test = nDisplayTest.parse_test(test_file)
+            if self.current_test is None:
+                LOGGER.error(f'Could not initialzie test data. Something is wrong with test file: "{test_file}"')
+
+    @classmethod
+    def plugin_settings(cls):
+        ''' Returns common settings that belong to all devices of this class. '''
+        return [
+            cls.tsettings['ndisplay_test_file']
+        ] + DevicenDisplay.plugin_settings()
+
+    @property
+    def test_name(self):
+        ''' Returns current test name '''
+        file_path = self.tsettings['ndisplay_test_file'].get_value()
+        file_name = os.path.splitext(os.path.basename(file_path))[0]
+        return file_name
+    
+    def reset_launch_temporary_data(self):
+        ''' Cleans up the temporary data used to launch a test '''
+        self.is_sending_cfg_file = False
+        self.target_cfg_file = None
+        self.is_sending_test_file = False
+        self.target_test_file = None
+
+    def device_widget_registered(self, device_widget):
+        '''[override]  Device interface method '''
+        # Initialize base first
+        super().device_widget_registered(device_widget)
+        # Initialize local signals
+        device_widget.signal_open_test_artifacts.connect(self.on_open_test_artifacts)
+
+    def on_open_test_artifacts(self):
+        ''' Opens test artifacts location in a browser '''
+        if self.artifacts_dir.exists():
+            os.startfile(self.artifacts_dir)
+        else:
+            LOGGER.warning('Output directory not found. Pleas try to open manually.')
+            
+    def get_connected_devices(self):
+        ''' Returns a list with the connected devices/nodes '''
+
+        # Find all nDisplayTest devices currently connected
+        nodes = []
+        for device in self.active_unreal_devices:
+            is_device_connected = (device.status == DeviceStatus.CLOSED) or (device.status == DeviceStatus.OPEN)
+            if (device.device_type == "nDisplayTest") and is_device_connected:
+                nodes.append(device)
+        return nodes
+
+    def on_send_file_complete(self, message):
+        ''' Handles file transition finish callback '''
+
+        # Get transfer result
+        try:
+            destination = message['destination']
+            succeeded = message['bAck']
+            error = message.get('error')
+        except KeyError:
+            LOGGER.error(f'Error parsing "send file complete" response ({message})')
+            return
+
+        # Based on the file extension, find what exactly we just transferred
+        ext = os.path.splitext(destination)[1].lower()
+
+        if (ext == '.ndisplay') and self.is_sending_cfg_file:
+            # Config file transfer finished
+            self.is_sending_cfg_file = False
+            if succeeded:
+                LOGGER.info(f"{self.name}: nDisplay config successfully transferred to {destination}")
+            else:
+                LOGGER.error(f"{self.name}: nDisplay config transfer failed: {error}")
+
+            # Now prepare and transfer the test file
+            self.generate_and_deploy_test_script()
+
+        elif (ext == '.xml') and self.is_sending_test_file:
+            # Test file transfer finished
+            self.is_sending_test_file = False
+            if succeeded:
+                LOGGER.info(f"{self.name}: nDisplay test file successfully transferred to {destination}")
+            else:
+                LOGGER.error(f"{self.name}: nDisplay test file transfer failed: {error}")
+            
+            # Finally, we can launch the test
+            self.launch_test()
+
+        else:
+            LOGGER.error(f"{self.name}: Unexpected send file completion for {destination}")
+            return
+
+    def generate_and_deploy_config_file(self):
+        ''' Generates temp config with local overrides '''
+
+        # Get original config file path
+        cfg_file = self.current_test.get_ndisplay_config_path()
+        # Config file extension
+        cfg_ext = os.path.splitext(cfg_file)[1].lower()
+
+        # Read config data
+        if cfg_ext == '.uasset':
+            cfg_content = self.__class__.extract_configexport_from_uasset(cfg_file).encode('utf-8')
+        else:
+            with open(cfg_file, 'rb') as f:
+                cfg_content = f.read()
+
+        # Apply local changes
+        cfg_ext = ".ndisplay" # only json configs for now
+        cfg_content = self.__class__.apply_local_overrides_to_config(cfg_content, cfg_ext).encode('utf-8')
+
+        # Generate target config file path
+        self.target_cfg_file = f"%TEMP%/ndisplay/{self.name}_{uuid.uuid4().hex[:12]}{cfg_ext}"
+
+        # Send content to the destination
+        self.is_sending_cfg_file = True
+        _, cfg_msg = message_protocol.create_send_filecontent_message(cfg_content, self.target_cfg_file)
+        self.unreal_client.send_message(cfg_msg)
+
+    def generate_and_deploy_test_script(self):
+        ''' Generate new test files with updated config path, and deploys it to a test machine '''
+
+        # Make sure we have generated target config path
+        if self.target_cfg_file is None:
+            LOGGER.error(f"No target config path was generated for {self.name}")
+            return
+
+        # Generate target test file
+        test_ext = ".xml"
+        self.target_test_file = f"%TEMP%/ndisplay/{self.name}_{uuid.uuid4().hex[:12]}{test_ext}"
+
+        try:
+            # Duplicate our test data
+            test_cloned = self.current_test.duplicate()
+
+            # Update config path in XML data
+            test_cloned.update_properties({
+                "DisplayConfigPath": f"{self.target_cfg_file}"
+            })
+
+            # Get test content
+            test_content = test_cloned.to_bytes()
+
+            # Send content to the destination
+            self.is_sending_test_file = True
+            _, test_msg = message_protocol.create_send_filecontent_message(test_content, self.target_test_file)
+            self.unreal_client.send_message(test_msg)
+
+        except Exception as e:
+            LOGGER.error(f"Could not deploy test file for {self.name}: {e}")
+            return
+
+    def prepare_location_for_artifacts(self):
+        ''' Generates the name of output directory, and creates it '''
+
+        try:
+            # Create directory with the generated name
+            formatted_test_time = datetime.now().strftime("%d%m%y_%I%M%S")
+            download_dir = self.get_log_download_dir()
+            download_dir = Path(download_dir) / self.test_name / formatted_test_time
+            Path(download_dir).mkdir(parents=True, exist_ok=True)
+            DevicenDisplayTest.artifacts_dir = Path(download_dir)
+        except:
+            DevicenDisplayTest.artifacts_dir = self.get_log_download_dir()
+
+        LOGGER.info(f'Test output dir: {str(DevicenDisplayTest.artifacts_dir)}')
+
+    def pre_launch_check(self) -> bool:
+        ''' Checks if everything is Ok before launch '''
+
+        # Check: test data
+        if self.current_test is None:
+            LOGGER.error("No test data loaded")
+            return False
+
+        # Check: test file
+        test_file = self.current_test.get_test_path()
+        if not os.path.exists(test_file):
+            LOGGER.error("Test file not found")
+            return False
+
+        # Check: at least one node available
+        nodes = self.get_connected_devices()
+        nodes_num = len(nodes)
+        if nodes_num < 1:
+            LOGGER.error(f"Invalid node count {nodes_num}")
+            return False
+
+        # Check: primary node is there
+        primary_name = DevicenDisplay.csettings['primary_device_name'].get_value()
+        if not any(node.name == primary_name for node in nodes):
+            LOGGER.error("Primary node not found")
+            return False
+
+        # Check Ok
+        return True
+
+    def launch(self, map_name):
+        ''' [override] Starts an instance of device '''
+
+        # Check if we're ready to go
+        if not self.pre_launch_check():
+            LOGGER.error("Pre-launch check failed")
+            return
+
+        LOGGER.info(f"Starting: {self.current_test.get_test_path()}")
+
+        # Primary device is responsible for preparing the output directory
+        if self.is_primary():
+            self.prepare_location_for_artifacts()
+            self.rsync_server.set_incoming_logs_path(DevicenDisplayTest.artifacts_dir)
+
+        try:
+            # Test launch is a multi-step procedure that starts from config deployment
+            self.generate_and_deploy_config_file()
+        except Exception as e:
+            print(f"Couldn't start test: {e}")
+            return
+
+    def launch_test(self):
+        ''' Starts the test on the destination machine '''
+
+        # Prepare some intermediate data to generate the launch parameters
+        runuat = os.path.join(str(Path(CONFIG.ENGINE_DIR.get_value()).parent), "RunUAT.bat")
+        project_name = CONFIG.PROJECT_NAME.get_value()
+        program_name = UnrealJobs.Unreal.value
+        script_dir="Engine/Plugins/VirtualProduction/ICVFXTesting/Build"
+
+        # Prepare the launch parameters
+        param_executable = 'cmd'
+        param_extra = self.generate_extra_cmdline_params()
+        param_args = '/C call ' + \
+            f'{runuat} ' + \
+            'BuildGraph ' + \
+            f'-Script=\"{self.target_test_file}\" ' + \
+            f'-Target=\"BuildAndTest {project_name}\" ' + \
+            f'-ScriptDir=\"{script_dir}\" ' + \
+            '-ubtargs="-NoMutex" ' + \
+            f'{param_extra}'
+
+        # Update last launch command
+        self.last_launch_command.update_value(f'{param_executable} {param_args}')
+
+        # Reset any temporaries of this session
+        self.reset_launch_temporary_data()
+
+        # Create process start request
+        puuid, msg = message_protocol.create_start_process_message(
+            prog_path = param_executable,
+            prog_args = param_args,
+            prog_name = program_name,
+            caller = self.name,
+            update_clients_with_stdout = False,
+            priority_modifier = sb_utils.PriorityModifier.Normal.value,
+            lock_gpu_clock = False,
+        )
+
+        # Schedule process start request to send
+        self.program_start_queue.add(
+            ProgramStartQueueItem(
+                name=program_name,
+                puuid_dependency=None,
+                puuid=puuid,
+                msg_to_unreal_client=msg,
+            ),
+            unreal_client=self.unreal_client,
+        )
+
+        # Optionally, minimize the window after launch
+        if DevicenDisplay.csettings['minimize_before_launch'].get_value(True):
+            self.minimize_windows()
+
+    def generate_extra_cmdline_params(self):
+        ''' Generates additional cmdline parameters for the test session '''
+
+        # Temporary & params
+        with_trace = True if CONFIG.INSIGHTS_TRACE_ENABLE.get_value() else False
+        win_pos = self.settings['window_position'].get_value()
+        win_res = self.settings['window_resolution'].get_value()
+        fullscreen = self.settings['fullscreen'].get_value()
+        headless = self.settings['headless'].get_value()
+
+        # Generate extra args
+        return '-set:IsOnStagePerfTest=true ' + \
+            f'-set:SkipPerfReportServer=true ' + \
+            f'-set:TestNameExt={self.test_name} ' + \
+            f'-set:WithTrace={with_trace} ' + \
+            f'-set:DisplayClusterNodeName={self.name} ' + \
+            f'-set:Fullscreen={fullscreen} ' + \
+            f'-set:Offscreen={headless} ' + \
+            f'-set:WndLocX={win_pos[0]} ' + \
+            f'-set:WndLocY={win_pos[1]} ' + \
+            f'-set:WndResX={win_res[0]} ' + \
+            f'-set:WndResY={win_res[1]} '
+
+    def generate_remote_test_root_directory(self):
+        ''' Generates remote test root directory path. '''
+
+        prj_name = CONFIG.PROJECT_NAME.get_value(self.name)
+
+        root_dir = Path(CONFIG.ENGINE_DIR.get_value(self.name)).parent
+        root_dir = root_dir / 'LocalBuilds/Reports'
+        root_dir = root_dir / prj_name
+        root_dir = root_dir / self.test_name
+        root_dir = root_dir / self.name
+        root_dir = root_dir / f'{prj_name}Win64{prj_name}Perf'
+
+        return root_dir
+
+    def generate_remote_test_log_filepath(self):
+        ''' Generates remote test log file path. '''
+        remote_file = self.generate_remote_test_root_directory()
+        remote_file = remote_file / 'ICVFXTest.PerformanceReport_(Win64_Development_EditorGame)'
+        remote_file = remote_file / 'EditorGame/EditorGameOutput.log'
+        return remote_file
+
+    def generate_remote_test_utrace_filepath(self):
+        ''' Generates remote test utrace file path. '''
+        remote_file = self.generate_remote_test_root_directory()
+        remote_file = remote_file / f'{self.name}_Base.utrace'
+        return remote_file
+
+    def generate_remote_test_report_dir(self):
+        ''' Generates remote test report directory path. '''
+        report_dir = self.generate_remote_test_root_directory()
+        report_dir = report_dir / 'ICVFXTest.PerformanceReport_(Win64_Development_EditorGame)/Reports/Performance'
+        return report_dir
+
+    def do_program_ended_update(self, *,
+        program_name: str,
+        returncode: int,
+        get_stdout_str: Callable[[], str]
+    ):
+        ''' [override] Handles test & transfer finish '''
+
+        # Let the base class do most of handling
+        super().do_program_ended_update(
+            program_name=program_name,
+            returncode=returncode,
+            get_stdout_str=get_stdout_str)
+
+        # Handle test finish
+        if program_name == UnrealJobs.Unreal.value:
+            self.start_retrieve_report(unreal_exit_code=returncode)
+            pass
+
+    def start_retrieve_log(self, unreal_exit_code: int):
+        ''' [override] Retrieve the log file if logging is enabled. '''
+
+        # Local & remote log file
+        local_log_path = DevicenDisplayTest.artifacts_dir / self.log_filename
+        remote_log_path = self.generate_remote_test_log_filepath()
+        # Update the most recent log reference
+        self.last_log_path.update_value(str(local_log_path))
+        # Download log file
+        return self.fetch_file(remote_log_path, self.log_filename)
+    
+    def start_retrieve_utrace(self, unreal_exit_code:int ):
+        ''' [override] Retrieve the utrace file if tracing is enabled. '''
+
+        if not CONFIG.INSIGHTS_TRACE_ENABLE.get_value():
+            return False
+
+        # Local & remote utrace file
+        local_utrace_filename = f'{self.name}.utrace'
+        local_utrace_path = DevicenDisplayTest.artifacts_dir / local_utrace_filename
+        remote_utrace_path = self.generate_remote_test_utrace_filepath()
+        # Update the most recent utrace reference
+        self.last_trace_path.update_value(str(local_utrace_path))
+        # Download utrace file
+        return self.fetch_file(remote_utrace_path, local_utrace_filename)
+
+    def start_retrieve_report(self, unreal_exit_code:int ):
+        ''' Retrieve report files. '''
+
+        # Reports to download
+        reports = [
+            'HistoricReport_14DaysBase.html',
+            'HistoricReport_AllTimeBase.html',
+            'SummaryTable.html'
+        ]
+
+        # Download all the reports
+        remote_report_dir = self.generate_remote_test_report_dir()
+        for report in reports:
+            self.fetch_file(Path(remote_report_dir) / report, f'{self.name}_{report}')
+
diff --git a/Engine/Plugins/VirtualProduction/Switchboard/Source/Switchboard/switchboard/devices/unreal/plugin_unreal.py b/Engine/Plugins/VirtualProduction/Switchboard/Source/Switchboard/switchboard/devices/unreal/plugin_unreal.py
--- a/Engine/Plugins/VirtualProduction/Switchboard/Source/Switchboard/switchboard/devices/unreal/plugin_unreal.py
+++ b/Engine/Plugins/VirtualProduction/Switchboard/Source/Switchboard/switchboard/devices/unreal/plugin_unreal.py
@@ -3087,7 +3087,12 @@ class DeviceUnreal(Device):
         else:
             return 'rsync'
 
-    def fetch_file(self, remote_path):
+    def fetch_file(self, remote_path, dest_filename = ""):
+        """
+        Copies a remote file to the local machine. Optionally renames
+        the file if a new name is provided.
+        """
+
         program_name = UnrealJobs.Retrieve.value
 
         rsync_path = self.get_rsync_path()
@@ -3098,7 +3103,7 @@ class DeviceUnreal(Device):
         dest_endpoint = \
             f'{SETTINGS.ADDRESS.get_value()}:{DeviceUnreal.rsync_server.port}'
         dest_module = DeviceUnreal.rsync_server.INCOMING_LOGS_MODULE
-        dest_path = f'rsync://{dest_endpoint}/{dest_module}/'
+        dest_path = f'rsync://{dest_endpoint}/{dest_module}/{dest_filename}'
 
         rsync_args = f'"{remote_cygpath}" "{dest_path}"'
 
