Switchboard: Fix first run when SWITCHBOARD_CONFIG_DIR override is set.

#rnx

diff --git a/Engine/Plugins/VirtualProduction/Switchboard/Source/Switchboard/switchboard/config.py b/Engine/Plugins/VirtualProduction/Switchboard/Source/Switchboard/switchboard/config.py
--- a/Engine/Plugins/VirtualProduction/Switchboard/Source/Switchboard/switchboard/config.py
+++ b/Engine/Plugins/VirtualProduction/Switchboard/Source/Switchboard/switchboard/config.py
@@ -3445,6 +3445,7 @@ class UserSettings(object):
             except ValueError as e:
                 data['config'] = str(self.CONFIG)
 
+        CONFIG_MGR.user_settings_file_path.parent.mkdir(parents=True, exist_ok=True)
         with open(CONFIG_MGR.user_settings_file_path, 'w', encoding='utf-8') as f:
             json.dump(data, f, indent=4)
 
