[Switchboard] [PerfTesting] Fixed an issue with remote paths that prevented the tests from starting

#rnx

diff --git a/Engine/Plugins/VirtualProduction/Switchboard/Source/Switchboard/switchboard/devices/ndisplayTest/plugin_ndisplayTest.py b/Engine/Plugins/VirtualProduction/Switchboard/Source/Switchboard/switchboard/devices/ndisplayTest/plugin_ndisplayTest.py
--- a/Engine/Plugins/VirtualProduction/Switchboard/Source/Switchboard/switchboard/devices/ndisplayTest/plugin_ndisplayTest.py
+++ b/Engine/Plugins/VirtualProduction/Switchboard/Source/Switchboard/switchboard/devices/ndisplayTest/plugin_ndisplayTest.py
@@ -280,9 +280,9 @@ class DevicenDisplayTest(DevicenDisplay):
     def reset_launch_temporary_data(self):
         ''' Cleans up the temporary data used to launch a test '''
         self.is_sending_cfg_file = False
-        self.target_cfg_file = None
+        self.remote_path_cfg = None
         self.is_sending_test_file = False
-        self.target_test_file = None
+        self.remote_path_test = None
 
     def device_widget_registered(self, device_widget):
         '''[override]  Device interface method '''
@@ -327,6 +327,8 @@ class DevicenDisplayTest(DevicenDisplay):
         if (ext == '.ndisplay') and self.is_sending_cfg_file:
             # Config file transfer finished
             self.is_sending_cfg_file = False
+            self.remote_path_cfg = destination
+
             if succeeded:
                 LOGGER.info(f"{self.name}: nDisplay config successfully transferred to {destination}")
             else:
@@ -338,6 +340,8 @@ class DevicenDisplayTest(DevicenDisplay):
         elif (ext == '.xml') and self.is_sending_test_file:
             # Test file transfer finished
             self.is_sending_test_file = False
+            self.remote_path_test = destination
+
             if succeeded:
                 LOGGER.info(f"{self.name}: nDisplay test file successfully transferred to {destination}")
             else:
@@ -370,32 +374,32 @@ class DevicenDisplayTest(DevicenDisplay):
         cfg_content = self.__class__.apply_local_overrides_to_config(cfg_content, cfg_ext).encode('utf-8')
 
         # Generate target config file path
-        self.target_cfg_file = f"%TEMP%/ndisplay/{self.name}_{uuid.uuid4().hex[:12]}{cfg_ext}"
+        remote_path_cfg_requested = f'%TEMP%/ndisplay/{self.name}_{uuid.uuid4().hex[:12]}{cfg_ext}'
 
         # Send content to the destination
         self.is_sending_cfg_file = True
-        _, cfg_msg = message_protocol.create_send_filecontent_message(cfg_content, self.target_cfg_file)
+        _, cfg_msg = message_protocol.create_send_filecontent_message(cfg_content, remote_path_cfg_requested)
         self.unreal_client.send_message(cfg_msg)
 
     def generate_and_deploy_test_script(self):
         ''' Generate new test files with updated config path, and deploys it to a test machine '''
 
-        # Make sure we have generated target config path
-        if self.target_cfg_file is None:
-            LOGGER.error(f"No target config path was generated for {self.name}")
+        # Make sure we have received the remote config path
+        if not self.remote_path_cfg:
+            LOGGER.error(f"No remote config path available for {self.name}")
             return
 
         # Generate target test file
         test_ext = ".xml"
-        self.target_test_file = f"%TEMP%/ndisplay/{self.name}_{uuid.uuid4().hex[:12]}{test_ext}"
+        remote_path_test_requested = f'%TEMP%/ndisplay/{self.name}_{uuid.uuid4().hex[:12]}{test_ext}'
 
         try:
             # Duplicate our test data
             test_cloned = self.current_test.duplicate()
 
-            # Update config path in XML data
+            # Update config path in the test XML data
             test_cloned.update_properties({
-                "DisplayConfigPath": f"{self.target_cfg_file}"
+                "DisplayConfigPath": f"{self.remote_path_cfg}"
             })
 
             # Get test content
@@ -403,7 +407,7 @@ class DevicenDisplayTest(DevicenDisplay):
 
             # Send content to the destination
             self.is_sending_test_file = True
-            _, test_msg = message_protocol.create_send_filecontent_message(test_content, self.target_test_file)
+            _, test_msg = message_protocol.create_send_filecontent_message(test_content, remote_path_test_requested)
             self.unreal_client.send_message(test_msg)
 
         except Exception as e:
@@ -461,6 +465,7 @@ class DevicenDisplayTest(DevicenDisplay):
         # Check if we're ready to go
         if not self.pre_launch_check():
             LOGGER.error("Pre-launch check failed")
+            self.widget._close()
             return
 
         LOGGER.info(f"Starting: {self.current_test.get_test_path()}")
@@ -475,6 +480,7 @@ class DevicenDisplayTest(DevicenDisplay):
             self.generate_and_deploy_config_file()
         except Exception as e:
             print(f"Couldn't start test: {e}")
+            self.widget._close()
             return
 
     def launch_test(self):
@@ -492,7 +498,7 @@ class DevicenDisplayTest(DevicenDisplay):
         param_args = '/C call ' + \
             f'{runuat} ' + \
             'BuildGraph ' + \
-            f'-Script=\"{self.target_test_file}\" ' + \
+            f'-Script=\"{self.remote_path_test}\" ' + \
             f'-Target=\"BuildAndTest {project_name}\" ' + \
             f'-ScriptDir=\"{script_dir}\" ' + \
             '-ubtargs="-NoMutex" ' + \
