Switchboard: Make rendering mode (e.g. mono) overrideable per node.

This is to support cases where the cluster needs to run some screens in stereo mode, others in mono.

#rnx

diff --git a/Engine/Plugins/VirtualProduction/Switchboard/Source/Switchboard/switchboard/devices/ndisplay/plugin_ndisplay.py b/Engine/Plugins/VirtualProduction/Switchboard/Source/Switchboard/switchboard/devices/ndisplay/plugin_ndisplay.py
--- a/Engine/Plugins/VirtualProduction/Switchboard/Source/Switchboard/switchboard/devices/ndisplay/plugin_ndisplay.py
+++ b/Engine/Plugins/VirtualProduction/Switchboard/Source/Switchboard/switchboard/devices/ndisplay/plugin_ndisplay.py
@@ -841,6 +841,7 @@ class DevicenDisplay(DeviceUnreal):
             DevicenDisplay.csettings['mediaprofile'],
             DevicenDisplay.csettings['graphics_adapter'],
             DevicenDisplay.csettings['packaged_game_path'],
+            DevicenDisplay.csettings['render_mode'],
             CONFIG.ENGINE_DIR,
             CONFIG.SOURCE_CONTROL_WORKSPACE,
             CONFIG.PROJECT_WORKSPACE,
