Remove reflect visibility in game cvar for users.  It causes a bunch of issues with game and editor clients.  A new workflow is supported in the level editor where users have more explicit control over level visibilty.

diff --git a/Engine/Plugins/Developer/Concert/ConcertMain/Source/ConcertClient/Public/ConcertClientSettings.h b/Engine/Plugins/Developer/Concert/ConcertMain/Source/ConcertClient/Public/ConcertClientSettings.h
--- a/Engine/Plugins/Developer/Concert/ConcertMain/Source/ConcertClient/Public/ConcertClientSettings.h
+++ b/Engine/Plugins/Developer/Concert/ConcertMain/Source/ConcertClient/Public/ConcertClientSettings.h
@@ -64,12 +64,8 @@ struct FConcertClientSettings
 	UPROPERTY(config, EditAnywhere, DisplayName="Latency Compensation", AdvancedDisplay, Category="Client Settings", meta=(ForceUnits=ms))
 	float LatencyCompensationMs;
 
-	/**
-	 * When level editor changes are made, reflect those changes to the game equivalent property.
-	 * This settings can be specified on the editor cmd with `-CONCERTREFLECTVISIBILITY=`.
-	 */
-	UPROPERTY(config, EditAnywhere, DisplayName="Reflect Level Visibility to Game", AdvancedDisplay, Category="Client Settings")
-	bool bReflectLevelEditorInGame = false;
+	UPROPERTY()
+	bool bReflectLevelEditorInGame_DEPRECATED = false;
 
 	/**
 	 * Enable extended version support when using Multi-user with precompiled and source builds.  When using Unreal Game
diff --git a/Engine/Plugins/Developer/Concert/ConcertSync/ConcertSyncClient/Source/ConcertSyncClient/Private/ConcertClientWorkspace.cpp b/Engine/Plugins/Developer/Concert/ConcertSync/ConcertSyncClient/Source/ConcertSyncClient/Private/ConcertClientWorkspace.cpp
--- a/Engine/Plugins/Developer/Concert/ConcertSync/ConcertSyncClient/Source/ConcertSyncClient/Private/ConcertClientWorkspace.cpp
+++ b/Engine/Plugins/Developer/Concert/ConcertSync/ConcertSyncClient/Source/ConcertSyncClient/Private/ConcertClientWorkspace.cpp
@@ -183,7 +183,6 @@ namespace UE::ConcertWorkspace::Private
 		bool bHasBeenSet = false;
 	};
 
-	static FCVarBool<TEXT("Editor.ReflectEditorLevelVisibilityWithGame"), true /*Editor Only*/> ReflectVisInGame;
 	static FCVarBool<TEXT("LevelInstance.ForceEditorWorldMode"), false /*Editor Only*/> LevelInstanceForceEditorWorldMode;
 	static FCVarBool<TEXT("EditorPaths.Enabled"), false /*Editor Only*/> EditorPaths;
 }
@@ -716,7 +715,6 @@ void FConcertClientWorkspace::HandleConnectionChanged(IConcertClientSession& InS
 		bHasSyncedWorkspace = false;
 		bFinalizeWorkspaceSyncRequested = false;
 		FConcertSlowTaskStackWorkaround::Get().PopTask(MoveTemp(InitialSyncSlowTask));
-		UE::ConcertWorkspace::Private::ReflectVisInGame.Reset();
 		UE::ConcertWorkspace::Private::LevelInstanceForceEditorWorldMode.Reset();
 		UE::ConcertWorkspace::Private::EditorPaths.Reset();
 	}
@@ -941,7 +939,6 @@ void FConcertClientWorkspace::OnEndFrame()
 	}
 	LiveSession->GetSessionDatabase().UpdateAsynchronousTasks();
 	IConcertClientRef ConcertClient = OwnerSyncClient->GetConcertClient();
-	UE::ConcertWorkspace::Private::ReflectVisInGame.SetValueOnce(ConcertClient->GetConfiguration()->ClientSettings.bReflectLevelEditorInGame);
 }
 
 void FConcertClientWorkspace::HandleWorkspaceSyncEndpointEvent(const FConcertSessionContext& Context, const FConcertWorkspaceSyncEndpointEvent& Event)
diff --git a/Engine/Plugins/Developer/Concert/ConcertSync/ConcertSyncClient/Source/ConcertSyncClient/Private/ConcertSyncClientModule.cpp b/Engine/Plugins/Developer/Concert/ConcertSync/ConcertSyncClient/Source/ConcertSyncClient/Private/ConcertSyncClientModule.cpp
--- a/Engine/Plugins/Developer/Concert/ConcertSync/ConcertSyncClient/Source/ConcertSyncClient/Private/ConcertSyncClientModule.cpp
+++ b/Engine/Plugins/Developer/Concert/ConcertSync/ConcertSyncClient/Source/ConcertSyncClient/Private/ConcertSyncClientModule.cpp
@@ -82,9 +82,6 @@ public:
 			ClientConfig->bRetryAutoConnectOnError |= FParse::Param(CommandLine, TEXT("CONCERTRETRYAUTOCONNECTONERROR"));
 			FParse::Bool(CommandLine, TEXT("-CONCERTRETRYAUTOCONNECTONERROR="), ClientConfig->bRetryAutoConnectOnError);
 
-			ClientConfig->ClientSettings.bReflectLevelEditorInGame |= FParse::Param(CommandLine, TEXT("CONCERTREFLECTVISIBILITY"));
-			FParse::Bool(CommandLine, TEXT("-CONCERTREFLECTVISIBILITY="), ClientConfig->ClientSettings.bReflectLevelEditorInGame);
-
 			ClientConfig->bIsHeadless |= FParse::Param(CommandLine, TEXT("CONCERTISHEADLESS"));
 			FParse::Bool(CommandLine, TEXT("-CONCERTISHEADLESS="), ClientConfig->bIsHeadless);
 			ClientConfig->bIsHeadless |= !FSlateNotificationManager::Get().AreNotificationsAllowed();
diff --git a/Engine/Plugins/Editor/ConsoleVariablesEditor/Source/ConsoleVariablesEditor/Public/ConsoleVariablesEditorProjectSettings.h b/Engine/Plugins/Editor/ConsoleVariablesEditor/Source/ConsoleVariablesEditor/Public/ConsoleVariablesEditorProjectSettings.h
--- a/Engine/Plugins/Editor/ConsoleVariablesEditor/Source/ConsoleVariablesEditor/Public/ConsoleVariablesEditorProjectSettings.h
+++ b/Engine/Plugins/Editor/ConsoleVariablesEditor/Source/ConsoleVariablesEditor/Public/ConsoleVariablesEditorProjectSettings.h
@@ -44,11 +44,6 @@ public:
 		PresetImportMode = EConsoleVariablesEditorPresetImportMode::AddToExisting;
 		
 		bAddAllChangedConsoleVariablesToCurrentPreset = true;
-
-		ChangedConsoleVariableSkipList =
-		{
-			"Editor.ReflectEditorLevelVisibilityWithGame"
-		};
 	}
 
 	/**
diff --git a/Engine/Plugins/VirtualProduction/Switchboard/Source/Switchboard/switchboard/devices/unreal/plugin_unreal.py b/Engine/Plugins/VirtualProduction/Switchboard/Source/Switchboard/switchboard/devices/unreal/plugin_unreal.py
--- a/Engine/Plugins/VirtualProduction/Switchboard/Source/Switchboard/switchboard/devices/unreal/plugin_unreal.py
+++ b/Engine/Plugins/VirtualProduction/Switchboard/Source/Switchboard/switchboard/devices/unreal/plugin_unreal.py
@@ -663,17 +663,6 @@ class DeviceUnreal(Device):
                 'If unset, defaults to $(ProjectDir)/Saved/Logs/Switchboard/'),
             category="General Settings",
         ),
-        'reflect_visibility_to_game': BoolSetting(
-            attr_name='reflect_visibility_to_game',
-            nice_name='Reflect Editor Visibility to Game',
-            value=True,
-            tool_tip=(
-                'Sets the value for `Reflect Level Visibilty to Game` for Multi-user editing. \n'
-                'Editor visibilty state will be applied to the game equivalent visibility properties. \n'
-                'This is useful for ICVFX workflows where the editor visibilty state directly \n'
-                'corresponds to the state on the render nodes.'),
-            category="Multi-User Settings",
-        ),
         'rsync_port': IntSetting(
             attr_name='rsync_port',
             nice_name='Rsync Server Port',
@@ -2110,10 +2099,6 @@ class DeviceUnreal(Device):
 
         return vproles, missing_roles
 
-    @property
-    def reflect_editor_vis_in_game(self) -> str:
-        return DeviceUnreal.csettings['reflect_visibility_to_game'].get_value()
-
     @property
     def udpmessaging_multicast_endpoint(self) -> str:
         device_multicast = DeviceUnreal.csettings[
@@ -2319,9 +2304,6 @@ class DeviceUnreal(Device):
                 self.name):
             command_line_args += ' -AutoDeclinePackageRecovery'
 
-        concert_vis_reflect = 1 if self.reflect_editor_vis_in_game else 0
-        command_line_args += f' -ConcertReflectVisibility={concert_vis_reflect}'
-
         # UdpMessaging endpoints
         if self.udpmessaging_multicast_endpoint:
             server_mc = CONFIG.MUSERVER_MULTICAST_ENDPOINT.get_value().strip()
diff --git a/Engine/Source/Editor/UnrealEd/Private/EditorLevelUtils.cpp b/Engine/Source/Editor/UnrealEd/Private/EditorLevelUtils.cpp
--- a/Engine/Source/Editor/UnrealEd/Private/EditorLevelUtils.cpp
+++ b/Engine/Source/Editor/UnrealEd/Private/EditorLevelUtils.cpp
@@ -61,13 +61,6 @@ DEFINE_LOG_CATEGORY(LogLevelTools);
 
 #define LOCTEXT_NAMESPACE "EditorLevelUtils"
 
-static TAutoConsoleVariable<int32>  CVarReflectEditorLevelVisibilityWithGame(
-	TEXT("Editor.ReflectEditorLevelVisibilityWithGame"),
-	0,
-	TEXT("Enables the transaction of game visibility state when editor visibility state changes.\n")
-	TEXT("0 - game state is *not* reflected with editor.\n")
-	TEXT("1 - game state is relfected with editor.\n"), ECVF_Default);
-
 UEditorLevelUtils::FCanMoveActorToLevelDelegate UEditorLevelUtils::CanMoveActorToLevelDelegate;
 UEditorLevelUtils::FOnMoveActorsToLevelEvent UEditorLevelUtils::OnMoveActorsToLevelEvent;
 
@@ -1225,7 +1218,6 @@ void SetLevelVisibilityNoGlobalUpdateInternal(ULevel* Level, const bool bShouldB
 	}
 	else
 	{
-		const bool bReflectVisibilityToGame = CVarReflectEditorLevelVisibilityWithGame.GetValueOnGameThread() != 0;
 		ULevelStreaming* StreamingLevel = NULL;
 		if (Level->OwningWorld == NULL || Level->OwningWorld->PersistentLevel != Level)
 		{
@@ -1249,10 +1241,6 @@ void SetLevelVisibilityNoGlobalUpdateInternal(ULevel* Level, const bool bShouldB
 
 			// Set the visibility state for this streaming level.
 			StreamingLevel->SetShouldBeVisibleInEditor(bShouldBeVisible);
-			if (bReflectVisibilityToGame)
-			{
-				StreamingLevel->SetShouldBeVisible(bShouldBeVisible);
-			}
 		}
 
 		ULayersSubsystem* Layers = GEditor->GetEditorSubsystem<ULayersSubsystem>();
@@ -1344,10 +1332,6 @@ void SetLevelVisibilityNoGlobalUpdateInternal(ULevel* Level, const bool bShouldB
 						Actor->UnregisterAllComponents();
 					}
 				}
-				if (bReflectVisibilityToGame)
-				{
-					Actor->SetActorHiddenInGame(Actor->bHiddenEdLevel);
-				}
 			}
 		}
 	}
