Switchboard: Suppress ObjectUE5Version warning for OS_SUB_OBJECT_SHADOW_SERIALIZATION, which still appears to deserialize correctly.

#rnx

diff --git a/Engine/Plugins/VirtualProduction/Switchboard/Source/Switchboard/switchboard/devices/unreal/uassetparser.py b/Engine/Plugins/VirtualProduction/Switchboard/Source/Switchboard/switchboard/devices/unreal/uassetparser.py
--- a/Engine/Plugins/VirtualProduction/Switchboard/Source/Switchboard/switchboard/devices/unreal/uassetparser.py
+++ b/Engine/Plugins/VirtualProduction/Switchboard/Source/Switchboard/switchboard/devices/unreal/uassetparser.py
@@ -314,7 +314,7 @@ class UassetParser(object):
 
         s.FileVersionLicenseeUE4 = self.readUInt32()
 
-        KNOWN_SUPPORTED_UE5VER = EUnrealEngineObjectUE5Version.PACKAGE_SAVED_HASH
+        KNOWN_SUPPORTED_UE5VER = 1017
         if s.FileVersionUE5 > KNOWN_SUPPORTED_UE5VER:
             LOGGER.warning(f"{self.f.name}: ObjectUE5Version {s.FileVersionUE5} too new; "
                            f"newest known supported version {KNOWN_SUPPORTED_UE5VER}")
