Switchboard: Fix "create config" not working anymore.

Regression due to 41362976 when support for alternate config paths was added.

#rnx

diff --git a/Engine/Plugins/VirtualProduction/Switchboard/Source/SwitchboardEditor/Private/SwitchboardEditorModule.cpp b/Engine/Plugins/VirtualProduction/Switchboard/Source/SwitchboardEditor/Private/SwitchboardEditorModule.cpp
--- a/Engine/Plugins/VirtualProduction/Switchboard/Source/SwitchboardEditor/Private/SwitchboardEditorModule.cpp
+++ b/Engine/Plugins/VirtualProduction/Switchboard/Source/SwitchboardEditor/Private/SwitchboardEditorModule.cpp
@@ -257,6 +257,10 @@ bool FSwitchboardEditorModule::CreateNewConfig(const FSwitchboardNewConfigUserOp
 		return false;
 	}
 	
+	UE_LOG(LogSwitchboardPlugin, Verbose, TEXT("Will launch Switchboard using script arguments: '%s' which were saved in: '%s'"),
+		*ScriptArgsString, *ScriptArgsFilepath
+	);
+
 	// The script that we will pass to Switchoard is in this plugin's Content/Python folder.
 	const FString PluginContentDir = IPluginManager::Get().FindPlugin(TEXT("Switchboard"))->GetContentDir();
 	const FString ScriptPath = UE::Switchboard::Private::ConcatPaths(PluginContentDir, "Python", "sb_script_new_config.py");
@@ -282,6 +286,8 @@ bool FSwitchboardEditorModule::RunProcess(const FString& InExe, const FString& I
 	const TCHAR* WorkingDirectory = nullptr;
 	void* PipeWriteChild = nullptr;
 
+	UE_LOG(LogSwitchboardPlugin, Display, TEXT("Launching: %s %s"), *InExe, *InArgs);
+
 	const FProcHandle Handle = FPlatformProcess::CreateProc(*InExe, *InArgs, bLaunchDetached, bLaunchHidden, bLaunchReallyHidden, OutProcessId, PriorityModifier, WorkingDirectory, PipeWriteChild);
 	return Handle.IsValid();
 }
