[Switchboard] PerfTesting upgrade

- Multi-nodes per machine (delayed launch)
- Editable test path textbox
- Explicit trace file name so there is no need to delete trace files if already exist
- Added ability to override game and editor targets in SB settings
- Added bool to enable/disable EditorBootTest (SB settings)
- Added bool to generate 24 fps reports (SB settings)
- Ability to set additional test info. By default it shows engine & project CL numbers
- Simplified test XML: MaxGPUCount, project name and path are automatically inherited from SB
- Explicit map override from SB
- Fixed boolean string to be lower case (BuildGraph uses lower case only)
- Fixed minor issues (None-checks)

#rnx

diff --git a/Engine/Plugins/VirtualProduction/ICVFXTesting/Build/Scripts/Automation/ICVFX.TestNode.cs b/Engine/Plugins/VirtualProduction/ICVFXTesting/Build/Scripts/Automation/ICVFX.TestNode.cs
--- a/Engine/Plugins/VirtualProduction/ICVFXTesting/Build/Scripts/Automation/ICVFX.TestNode.cs
+++ b/Engine/Plugins/VirtualProduction/ICVFXTesting/Build/Scripts/Automation/ICVFX.TestNode.cs
@@ -140,6 +140,12 @@ namespace ICVFXTest
 		[AutoParamWithNames(1, "ICVFXTest.MaxGPUCount")]
 		public int MaxGPUCount;
 
+		/// <summary>
+		/// Report type
+		/// </summary>
+		[AutoParamWithNames("ICVFXTest.SummaryReportType")]
+		public string SummaryReportType;
+
 		/// <summary>
 		/// Path to the exported ndisplay config to use for testing.
 		/// </summary>
@@ -172,7 +178,13 @@ namespace ICVFXTest
 		/// </summary>
 		[AutoParamWithNames("ICVFXTest.TraceRootFolder")]
 		public string TraceRootFolder;
-		
+
+		/// <summary>
+		/// Force file name to use for a trace file. If not specified, the default one will be generated.
+		/// </summary>
+		[AutoParamWithNames("ICVFXTest.TraceFileName")]
+		public string TraceFileName;
+
 		/// <summary>
 		///  Path to the folder where perf cache data should be stored.
 		/// </summary>
diff --git a/Engine/Plugins/VirtualProduction/ICVFXTesting/Build/Scripts/Automation/Tests/ICVFXTest.AutoTest.cs b/Engine/Plugins/VirtualProduction/ICVFXTesting/Build/Scripts/Automation/Tests/ICVFXTest.AutoTest.cs
--- a/Engine/Plugins/VirtualProduction/ICVFXTesting/Build/Scripts/Automation/Tests/ICVFXTest.AutoTest.cs
+++ b/Engine/Plugins/VirtualProduction/ICVFXTesting/Build/Scripts/Automation/Tests/ICVFXTest.AutoTest.cs
@@ -1,12 +1,8 @@
 // Copyright Epic Games, Inc. All Rights Reserved.
 
-using EpicGame;
 using Gauntlet;
 using Microsoft.Extensions.Logging;
 using System.Collections.Generic;
-using System.Diagnostics;
-using System.IO;
-using UnrealBuildTool;
 
 using Log = EpicGames.Core.Log;
 
@@ -83,12 +79,17 @@ namespace ICVFXTest
 
 			if (Config.EnableTrace && Config.TraceRootFolder.Length != 0)
 			{
-				string ClusterNode = Config.DisplayClusterNodeName;
-				string TraceFile = string.IsNullOrEmpty(ClusterNode) ?
-					$"{Config.TraceRootFolder}/{GetTestSuffix()}.utrace" :
-					$"{Config.TraceRootFolder}/{ClusterNode}_{GetTestSuffix()}.utrace";
-
-				ClientRole.CommandLineParams.AddOrAppendParamValue("execcmds", "ICVFXTest.TraceFileName " + $"{TraceFile}");
+				// If no file name specified, generate default name
+				if (string.IsNullOrEmpty(Config.TraceFileName))
+				{
+					// Original approach
+					ClientRole.CommandLineParams.AddOrAppendParamValue("execcmds", "ICVFXTest.TraceFileName " + $"{Config.TraceRootFolder}/{GetTestSuffix()}.utrace");
+				}
+				else
+				{
+					// Explicit file name
+					ClientRole.CommandLineParams.AddOrAppendParamValue("execcmds", "ICVFXTest.TraceFileName " + $"{Config.TraceRootFolder}/{Config.TraceFileName}");
+				}
 			}
 
 			ClientRole.CommandLineParams.AddOrAppendParamValue("execcmds", "r.nanite " + (UseNanite() ? 1 : 0));
diff --git a/Engine/Plugins/VirtualProduction/ICVFXTesting/Build/Scripts/Automation/Tests/ICVFXTest.PerformanceReport.cs b/Engine/Plugins/VirtualProduction/ICVFXTesting/Build/Scripts/Automation/Tests/ICVFXTest.PerformanceReport.cs
--- a/Engine/Plugins/VirtualProduction/ICVFXTesting/Build/Scripts/Automation/Tests/ICVFXTest.PerformanceReport.cs
+++ b/Engine/Plugins/VirtualProduction/ICVFXTesting/Build/Scripts/Automation/Tests/ICVFXTest.PerformanceReport.cs
@@ -89,12 +89,33 @@ namespace ICVFXTest
 	/// </summary>
 	public class PerformanceReport : AutoTest
 	{
-		private DirectoryInfo TempPerfCSVDir => new DirectoryInfo(Path.Combine(Unreal.RootDirectory.FullName, "GauntletTemp", "PerfReportCSVs_" + GetConfiguration().DisplayClusterNodeName));
 		private static ILogger Logger => Log.Logger;
 		private string OriginalBuildName;
 		private string OverrideDisplayConfigPath;
 		private string OverrideDisplayClusterNode;
 
+		// Temporary directory for the CSV file
+		private DirectoryInfo TempPerfCSVDir
+		{
+			get
+			{
+				// When running with node ID specified, embed the ID into the directory name
+				string NodeId = GetConfiguration().DisplayClusterNodeName;
+				string ReportDirName = string.IsNullOrEmpty(NodeId) ? "PerfReportCSVs" : "PerfReportCSVs_" + NodeId;
+				return new DirectoryInfo(Path.Combine(Unreal.RootDirectory.FullName, "GauntletTemp", ReportDirName));
+			}
+		}
+
+		// Generates cmdline parameter to configure the neccessary type of report (single line embedding)
+		private string ReportTypeCmdLineParam
+		{
+			get
+			{
+				string CfgReportType = GetConfiguration().SummaryReportType;
+				return string.IsNullOrEmpty(CfgReportType) ? string.Empty : $"-reportType {CfgReportType}";
+			}
+		}
+
 		public PerformanceReport(UnrealTestContext InContext)
 			: base(InContext)
 		{
@@ -138,7 +159,7 @@ namespace ICVFXTest
 		}
 
 		protected override void InitHandledErrors()
-        {
+		{
 			base.InitHandledErrors();
 		}
 
@@ -333,8 +354,22 @@ namespace ICVFXTest
 			// Win64 is actually called "Windows" in csv profiles
 			var PlatformNameFilter = Platform == UnrealTargetPlatform.Win64 ? "Windows" : $"{Platform}";
 
+			// Generate HTML report
+			var Args = new[]
+			{
+				$"-csvdir \"{NewestDir}\"",
+				$"-o \"{ReportPath}\"",
+				$"-reportxmlbasedir \"{ReportConfigDir}\"",
+				$"-summaryTableCache \"{ReportCacheDir}\"",
+				$"-metadatafilter platform=\"{PlatformNameFilter}\"",
+				$"{ReportTypeCmdLineParam}",
+				"-searchpattern csvprofile*"
+			};
+
+			var ArgStr = string.Join(" ", Args);
+
 			// Produce the detailed report, and update the perf cache
-			CommandUtils.RunAndLog(ToolPath.FullName, $"-csvdir \"{NewestDir}\" -o \"{ReportPath}\" -reportxmlbasedir \"{ReportConfigDir}\" -summaryTableCache \"{ReportCacheDir}\" -searchpattern csvprofile* -metadatafilter platform=\"{PlatformNameFilter}\"", out int ErrorCode);
+			CommandUtils.RunAndLog(ToolPath.FullName, ArgStr, out int ErrorCode);
 			if (ErrorCode != 0)
 			{
 				Logger.LogError($"PerfReportTool returned error code \"{ErrorCode}\" while generating detailed report.");
@@ -351,6 +386,7 @@ namespace ICVFXTest
 					$"-reportxmlbasedir \"{ReportConfigDir}\"",
 					$"-o \"{base.GetConfiguration().SummaryReportPath}\"",
 					$"-metadatafilter \"{string.Join(" and ", Filter)}\"",
+					$"{ReportTypeCmdLineParam}",
 					"-summaryTable autoPerfReportStandard",
 					"-condensedSummaryTable autoPerfReportStandard",
 					"-emailtable",
@@ -374,6 +410,7 @@ namespace ICVFXTest
 					$"-o \"{base.GetConfiguration().SummaryReportPath}\"",
 					$"-metadatafilter \"{string.Join(" and ", Filter)}\"",
 					$"-summaryTableFilename \"{Name  + GetTestSuffix()}\"",
+					$"{ReportTypeCmdLineParam}",
 					"-csvTable",
 					"-summaryTable autoPerfReportStandard",
 					"-condensedSummaryTable autoPerfReportStandard",
@@ -400,6 +437,7 @@ namespace ICVFXTest
 					$"-reportxmlbasedir \"{ReportConfigDir}\"",
 					$"-o \"{ReportPath}\"",
 					$"-metadatafilter \"{string.Join(" and ", Filter)}\"",
+					$"{ReportTypeCmdLineParam}",
 					"-summaryTable autoPerfReportStandard",
 					"-condensedSummaryTable autoPerfReportStandard",
 					"-emailtable",
diff --git a/Engine/Plugins/VirtualProduction/ICVFXTesting/Build/Scripts/PerfReport/ReportTypes.xml b/Engine/Plugins/VirtualProduction/ICVFXTesting/Build/Scripts/PerfReport/ReportTypes.xml
--- a/Engine/Plugins/VirtualProduction/ICVFXTesting/Build/Scripts/PerfReport/ReportTypes.xml
+++ b/Engine/Plugins/VirtualProduction/ICVFXTesting/Build/Scripts/PerfReport/ReportTypes.xml
@@ -11,6 +11,7 @@
 		<mapping statName="MemoryFreeMB" displayName="MemFreeMB"/>
 		<mapping statName="PhysicalUsedMB" displayName="MemUsedMB"/>
 		<mapping statName="buildversion" displayName="Build Version"/>
+		<mapping statName="testbuildversion" displayName="Build Info"/>
 		<mapping statName="config" displayName="Config"/>
 		<mapping statName="commandline" displayName="Commandline"/>
 		<mapping statName="syncLoads" displayName="Sync Loads"/>
@@ -84,11 +85,11 @@
 	<!-- The preset Summary Tables below can be used to provide a filtered and sorted view of that data (see -summaryTable). -->
 	<summaryTables>
 		<summaryTable name="all" rowSort="buildversion,Platform,deviceprofile">
-			<filter>buildversion,Platform,DeviceProfile,scalabilitylevel,*</filter>
+			<filter>buildversion,testbuildversion,Platform,DeviceProfile,scalabilitylevel,*</filter>
 		</summaryTable>
 
-		<summaryTable name="default" rowSort="buildversion,playlistid,Platform,deviceprofile,scalabilitylevel,TargetFPS">
-			<filter>Platform,DeviceProfile,scalabilitylevel,Report,TargetFPS,MVP*,Hitches/Min,HitchTimePercent,FrameTime*,GamethreadTime*,Renderthreadtime*,RHIThreadTime*,GPUTime*,MemoryFreeMB Min,PhysicalUsedMB Max,DynamicResolution*,RHI/Drawcalls*,Basic/TicksQueued*,Total Time*,Csv File,buildversion,config,playlistid,framecount*,syncloads</filter>
+		<summaryTable name="default" rowSort="buildversion,testbuildversion,playlistid,Platform,deviceprofile,scalabilitylevel,TargetFPS">
+			<filter>Platform,DeviceProfile,scalabilitylevel,Report,TargetFPS,MVP*,Hitches/Min,HitchTimePercent,FrameTime*,GamethreadTime*,Renderthreadtime*,RHIThreadTime*,GPUTime*,MemoryFreeMB Min,PhysicalUsedMB Max,DynamicResolution*,RHI/Drawcalls*,Basic/TicksQueued*,Total Time*,Csv File,buildversion,testbuildversion,testbuildversion,config,playlistid,framecount*,syncloads</filter>
 		</summaryTable>
 
 		<summaryTable name="memory" rowSort="buildversion,Platform,deviceprofile,scalabilitylevel">
@@ -100,12 +101,12 @@
 		</summaryTable>
 
 		<summaryTable name="historical" rowSort="buildversion,platform,deviceprofile">
-			<filter>buildversion,Platform,DeviceProfile,scalabilitylevel, config,TargetFPS,testname,Report,MVP*,Hitches/Min,HitchTimePercent,FrameTime*,GamethreadTime*,Renderthreadtime*,RHIThreadTime*,GPUTime*,MemoryFreeMB Min,PhysicalUsedMB Max,DynamicResolution*,RHI/Drawcalls*,Ticks/PlayerPawns*,Basic/TicksQueued*,Total Time*,framecount*,fmsglogf/fmsglogfcount avg,syncloads</filter>
+			<filter>buildversion,testbuildversion,Platform,DeviceProfile,scalabilitylevel, config,TargetFPS,testname,Report,MVP*,Hitches/Min,HitchTimePercent,FrameTime*,GamethreadTime*,Renderthreadtime*,RHIThreadTime*,GPUTime*,MemoryFreeMB Min,PhysicalUsedMB Max,DynamicResolution*,RHI/Drawcalls*,Ticks/PlayerPawns*,Basic/TicksQueued*,Total Time*,framecount*,fmsglogf/fmsglogfcount avg,syncloads</filter>
 		</summaryTable>
 
 		<!-- This table is used to produce the historic reports at the end of each Frosty.PerformanceReport Gauntlet test. -->
-		<summaryTable name="autoPerfReportStandard" rowSort="buildversion,Platform,deviceprofile" scrollableFormatting="1" reverseSortRows="1">
-			<filter>buildversion,DeviceProfile,config,count,testname,TargetFPS,MVP*,Hitches/Min,HitchTimePercent,FrameTime*,GamethreadTime*,Renderthreadtime*,RHIThreadTime*,GPUTime*,MemoryFreeMB Min,PhysicalUsedMB Max,DynamicResolution*,Total Time*,RHI/*,Basic/TicksQueued*,Ticks/Vehicles*,framecount*,Csv File,Report</filter>
+		<summaryTable name="autoPerfReportStandard" rowSort="buildversion,testbuildversion,Platform,deviceprofile" scrollableFormatting="1" reverseSortRows="1">
+			<filter>buildversion,testbuildversion,DeviceProfile,config,count,testname,TargetFPS,MVP*,Hitches/Min,HitchTimePercent,FrameTime*,GamethreadTime*,Renderthreadtime*,RHIThreadTime*,GPUTime*,MemoryFreeMB Min,PhysicalUsedMB Max,DynamicResolution*,Total Time*,RHI/*,Basic/TicksQueued*,Ticks/Vehicles*,framecount*,Csv File,Report</filter>
 		</summaryTable>
 
 	</summaryTables>
@@ -119,7 +120,7 @@
 				<requiredmetadata key="targetframerate" allowedValues="60" stopIfKeyFound="1" ignoreIfKeyNotFound="1"/>
 			</autodetection>
 			
-			<metadataToShow>buildversion,platform,deviceprofile,csvfilename,rhiname,desktophwinfo,testname</metadataToShow>
+			<metadataToShow>buildversion,testbuildversion,platform,deviceprofile,csvfilename,rhiname,desktophwinfo,testname</metadataToShow>
 
 			<summary type="fpschart" fps="60" hitchThreshold="60" useEngineHitchMetric="1">
 				<stats>Frametime,GameThreadtime,RenderThreadtime,RHIThreadTime,GPUtime,RHI/Drawcalls,DynamicResolutionPercentage,MemoryFreeMB(min),PhysicalUsedMB(max),Basic/TicksQueued,fmsglogf/fmsglogfcount</stats>
@@ -181,5 +182,74 @@
 			<graph title="Log Counts" budget="0" />
 			<graph title="CSV Profiler" budget="0" />
 		</reporttype>
+
+		<!-- VP 24fps report type -->
+		<reporttype name="VP24fps" title="VP 24FPS Performance Report" vsync="0">
+			<autodetection requiredstats="frametime">
+				<requiredmetadata key="targetframerate" allowedValues="24" stopIfKeyFound="1" ignoreIfKeyNotFound="1"/>
+			</autodetection>
+			
+			<metadataToShow>buildversion,testbuildversion,platform,deviceprofile,csvfilename,rhiname,desktophwinfo,testname</metadataToShow>
+
+			<summary type="fpschart" fps="24" hitchThreshold="24" useEngineHitchMetric="1">
+				<stats>Frametime,GameThreadtime,RenderThreadtime,RHIThreadTime,GPUtime,RHI/Drawcalls,DynamicResolutionPercentage,MemoryFreeMB(min),PhysicalUsedMB(max),Basic/TicksQueued,fmsglogf/fmsglogfcount</stats>
+				<colourThresholds stat="Hitches/Min"> 2, 3, 5, 6</colourThresholds>
+				<colourThresholds stat="HitchTimePercent"> 0.1, 0.5, 1, 2</colourThresholds>
+				<colourThresholds stat="MVP24"> 2, 3, 4, 5</colourThresholds>
+				<colourThresholds stat="MemoryFreeMB(min)">400,300,200,50</colourThresholds>
+				<colourThresholds stat="DynamicResolutionPercentage">90,80,70,60</colourThresholds>
+				<colourThresholds stat="Frametime"> 41.66, 44, 48, 56 </colourThresholds>
+				<colourThresholds stat="GameThreadtime"> 41.66, 44, 48, 56 </colourThresholds>
+				<colourThresholds stat="RHIThreadTime"> 41.66, 44, 48, 56 </colourThresholds>
+				<colourThresholds stat="RenderThreadtime"> 41.66, 44, 48, 56 </colourThresholds>
+				<colourThresholds stat="GPUtime"> 41.66, 44, 48, 56 </colourThresholds>
+				<colourThresholds stat="Basic/TicksQueued"> 150, 200, 250, 300 </colourThresholds>
+				<colourThresholds stat="fmsglogf/fmsglogfcount">0.01,0.1,0.2,0.5</colourThresholds>
+			</summary>
+
+			<summary type="hitches">
+				<stats>Frametime,GameThreadtime,RenderThreadtime,RHIThreadTime,GPUtime</stats>
+				<colourThresholds>82, 55, 41.6, 38.0</colourThresholds>
+				<hitchThresholds>60, 150, 250, 500, 750, 1000, 2000</hitchThresholds>
+			</summary>
+
+			<!--
+			<summaryRef name="mapOverlay"/>
+			<summaryRef name="syncLoads"/>
+			<summaryRef name="textureStreaming"/>
+			<summaryRef name="psos"/>
+			-->
+
+			<graph title="Stat Unit" budget="41.66" />
+			<graph title="Stat Unit Raw" budget="41.66" />
+			<graph title="Gamethread Breakdown" budget="41.66"/>
+			<graph title="Gamethread Breakdown Raw" budget="41.66"/>
+			<graph title="Gamethread Waits" budget="2"/>
+			<graph title="Gamethread Physics" budget="2"/>
+			<graph title="Gamethread Breakdown Giant" budget="41.66" external="1"/>
+			<graph title="Ticks Breakdown" budget="200"/>
+			<graph title="Actor Counts" budget="0"/>
+			<graph title="Cmds DayaLayers" budget="300"/>
+			<graph title="Physical Memory" budget="300"/>
+			<graph title="Physical Memory Over Budget" budget="300"/>
+			<graph title="Dynamic Resolution" budget="80"/>
+			<graph title="GPU Breakdown" budget="41.66" />
+			<graph title="GPU Breakdown Raw" budget="41.66" />
+			<graph title="D3D12 Video Memory"  budget="0" />
+			<graph title="Niagara Breakdown" budget="41.66"/>
+			<graph title="Light Counts" budget="0" />
+			<graph title="RHI Drawcalls" budget="3000" />
+			<graph title="Renderthread Breakdown" budget="41.66"/>
+			<graph title="Renderthread Breakdown Raw" budget="41.66"/>
+			<graph title="FScene Breakdown" budget="41.66"/>
+			<graph title="FScene Breakdown Raw" budget="41.66"/>
+			<graph title="Render Targets Peak(MB)" budget="0" />
+			<graph title="Textures Peak(MB)" budget="0" />
+			<graph title="Virtual Texture" budget="0" />
+			<graph title="Animation Budget" budget="1.5" />
+			<graph title="Texture Streaming" budget="0" />
+			<graph title="Log Counts" budget="0" />
+			<graph title="CSV Profiler" budget="0" />
+		</reporttype>
 	</reporttypes>
 </root>
\ No newline at end of file
diff --git a/Engine/Plugins/VirtualProduction/ICVFXTesting/Scripts/GauntletSettings.xml b/Engine/Plugins/VirtualProduction/ICVFXTesting/Scripts/GauntletSettings.xml
--- a/Engine/Plugins/VirtualProduction/ICVFXTesting/Scripts/GauntletSettings.xml
+++ b/Engine/Plugins/VirtualProduction/ICVFXTesting/Scripts/GauntletSettings.xml
@@ -15,12 +15,15 @@
 
 	<!-- For on-set testing, there will be custom approach -->
 	<Option Name="TestNameExt" DefaultValue="" Description="Extended test name (sub-test name for on-stage runs)" />
+	<Option Name="TraceFileName" DefaultValue="" Description="Explicit trace file name" />
+	<Option Name="TestBuildName" DefaultValue="CL-0" Description="Explicit build name" />
 	<Option Name="Fullscreen" Restrict="true|false" DefaultValue="false" Description="Should run fullscreen" />
 	<Option Name="Offscreen" Restrict="true|false" DefaultValue="false" Description="Should run offscreen" />
 	<Option Name="WndLocX" DefaultValue="0" Description="Window X position" />
 	<Option Name="WndLocY" DefaultValue="0" Description="Window Y position" />
 	<Option Name="WndResX" DefaultValue="1920" Description="Window X size" />
 	<Option Name="WndResY" DefaultValue="1080" Description="Window Y size" />
+	<Option Name="SummaryReportType" DefaultValue="" Description="The type/preset of the final test reports" />
 	<Option Name="IsOnStagePerfTest" Restrict="true|false" DefaultValue="false" Description="Whether it's running on a real stage" />
 	<Do If="$(IsOnStagePerfTest)">
 		<Property Name="EditorTestList" Value="ICVFXTest.PerformanceReport+" If="$(WithPerformanceReportTest)" />
@@ -89,33 +92,44 @@
 	
 	<Property Name="TraceRootFolder" Value="$(TestOutputDir)/$(BuildNamePath)" />
 	<Property Name="TraceRootFolder" Value="$(TestOutputDir)" If="$(IsOnStagePerfTest)" />	
+
+	<Property Name="BuildName" Value="$(EscapedBranch)-CL-$(Change)$(PreflightSuffix)" />
+	<Property Name="BuildName" Value="$(TestBuildName)" If="$(IsOnStagePerfTest)" />	
 		
 	<Property Name="ExtraGauntletCommonArgs" Value="-screenshotperiod=0 -ICVFXTest.StatCommands -ICVFXTest.MaxRunCount=$(MaxRunCount) $(LogOutputParams)" />
 	<Property Name="ExtraGauntletEditorTestArgs" Value="-game" />
 
 	<Property Name="ExtraGauntletEditorTestArgs" Value="$(ExtraGauntletEditorTestArgs) -ICVFXTest.IsBuildMachine=$(IsBuildMachine)" />
 	<Property Name="ExtraGauntletEditorTestArgs" Value="$(ExtraGauntletEditorTestArgs) -ICVFXTest.TestName=$(ProjectTestName)" />
-	<Property Name="ExtraGauntletEditorTestArgs" Value="$(ExtraGauntletEditorTestArgs) -ICVFXTest.MapOverride=$(MapOverride) -ICVFXTest.DisplayConfig=$(DisplayConfigPath) -ICVFXTest.ViewportScreenPercentage=$(ViewportScreenPercentage)" />
-	<Property Name="ExtraGauntletEditorTestArgs" Value="$(ExtraGauntletEditorTestArgs) -ICVFXTest.Fullscreen=$(Fullscreen) -ICVFXTest.Offscreen=$(Offscreen)" />
-	<Property Name="ExtraGauntletEditorTestArgs" Value="$(ExtraGauntletEditorTestArgs) -ICVFXTest.WndLocX=$(WndLocX) -ICVFXTest.WndLocY=$(WndLocY) -ICVFXTest.WndResX=$(WndResX) -ICVFXTest.WndResY=$(WndResY)" />
+	<Property Name="ExtraGauntletEditorTestArgs" Value="$(ExtraGauntletEditorTestArgs) -ICVFXTest.MapOverride=$(MapOverride)" />
+	<Property Name="ExtraGauntletEditorTestArgs" Value="$(ExtraGauntletEditorTestArgs) -ICVFXTest.DisplayConfig=$(DisplayConfigPath)" />
+	<Property Name="ExtraGauntletEditorTestArgs" Value="$(ExtraGauntletEditorTestArgs) -ICVFXTest.ViewportScreenPercentage=$(ViewportScreenPercentage)" />
+	<Property Name="ExtraGauntletEditorTestArgs" Value="$(ExtraGauntletEditorTestArgs) -ICVFXTest.Fullscreen=$(Fullscreen)" />
+	<Property Name="ExtraGauntletEditorTestArgs" Value="$(ExtraGauntletEditorTestArgs) -ICVFXTest.Offscreen=$(Offscreen)" />
+	<Property Name="ExtraGauntletEditorTestArgs" Value="$(ExtraGauntletEditorTestArgs) -ICVFXTest.WndLocX=$(WndLocX)" />
+	<Property Name="ExtraGauntletEditorTestArgs" Value="$(ExtraGauntletEditorTestArgs) -ICVFXTest.WndLocY=$(WndLocY)" />
+	<Property Name="ExtraGauntletEditorTestArgs" Value="$(ExtraGauntletEditorTestArgs) -ICVFXTest.WndResX=$(WndResX)" />
+	<Property Name="ExtraGauntletEditorTestArgs" Value="$(ExtraGauntletEditorTestArgs) -ICVFXTest.WndResY=$(WndResY)" />
 	<Property Name="ExtraGauntletEditorTestArgs" Value="$(ExtraGauntletEditorTestArgs) -ICVFXTest.DisplayClusterNodeName=$(DisplayClusterNodeName)" />
 	<Property Name="ExtraGauntletEditorTestArgs" Value="$(ExtraGauntletEditorTestArgs) -ICVFXTest.D3DDebug" If="$(WithD3DDebug)"/>
 	<Property Name="ExtraGauntletEditorTestArgs" Value="$(ExtraGauntletEditorTestArgs) -ICVFXTest.GPUCrashDebugging" If="$(WithGPUCrashDebugging)"/>
 	<Property Name="ExtraGauntletEditorTestArgs" Value="$(ExtraGauntletEditorTestArgs) -ICVFXTest.Nanite=$(WithNanite)"/>
 	<Property Name="ExtraGauntletEditorTestArgs" Value="$(ExtraGauntletEditorTestArgs) -ICVFXTest.Lumen=$(WithLumen)"/>
 	<Property Name="ExtraGauntletEditorTestArgs" Value="$(ExtraGauntletEditorTestArgs) -ICVFXTest.MaxGPUCount=$(MaxGPUCount)"/>
+	<Property Name="ExtraGauntletEditorTestArgs" Value="$(ExtraGauntletEditorTestArgs) -ICVFXTest.SummaryReportType=$(SummaryReportType)" If="'$(SummaryReportType)' != ''" />
 	<Property Name="ExtraGauntletEditorTestArgs" Value="$(ExtraGauntletEditorTestArgs) -ICVFXTest.SummaryReportPath=$(SummaryReportPath)"/>
 	<Property Name="ExtraGauntletEditorTestArgs" Value="$(ExtraGauntletEditorTestArgs) -ICVFXTest.RDGImmediate" If="$(WithRDGImmediate)"/>
 	<Property Name="ExtraGauntletEditorTestArgs" Value="$(ExtraGauntletEditorTestArgs) -ICVFXTest.NoRHIThread" If="$(WithNoRHIThread)"/>
 	<Property Name="ExtraGauntletEditorTestArgs" Value="$(ExtraGauntletEditorTestArgs) -ICVFXTest.DisableVirtualShadowMaps" If="!$(WithVirtualShadowMaps)"/>
 	<Property Name="ExtraGauntletEditorTestArgs" Value="$(ExtraGauntletEditorTestArgs) -ICVFXTest.EnableTrace" If="$(WithTrace)"/>
 	<Property Name="ExtraGauntletEditorTestArgs" Value="$(ExtraGauntletEditorTestArgs) -ICVFXTest.TraceRootFolder=$(TraceRootFolder)" If="$(WithTrace)"/>
+	<Property Name="ExtraGauntletEditorTestArgs" Value="$(ExtraGauntletEditorTestArgs) -ICVFXTest.TraceFileName=$(TraceFileName)" If="'$(TraceFileName)' != ''"/>
 	<Property Name="ExtraGauntletEditorTestArgs" Value="$(ExtraGauntletEditorTestArgs) -ICVFXTest.PerfCacheFolder=$(ReportCacheDirRoot)"/>
 	<Property Name="ExtraGauntletEditorTestArgs" Value="$(ExtraGauntletEditorTestArgs) -ICVFXTest.NoLocalReports" If="$(NoLocalReports)"/>
 	<Property Name="ExtraGauntletEditorTestArgs" Value="$(ExtraGauntletEditorTestArgs) -ICVFXTest.SkipPerfReportServer" If="$(SkipPerfReportServer)"/>
 	<Property Name="ExtraGauntletEditorTestArgs" Value="$(ExtraGauntletEditorTestArgs) -ICVFXTest.IsDevelopment" If="$(IsDevelopment)"/>
 	<Property Name="ExtraGauntletEditorTestArgs" Value="$(ExtraGauntletEditorTestArgs) -ICVFXTest.PerfReportServerImportDir=$(PerfReportServerImportDir)"/>
-	<Property Name="ExtraGauntletEditorTestArgs" Value="$(ExtraGauntletEditorTestArgs) -ICVFXTest.BuildName=$(EscapedBranch)-CL-$(Change)$(PreflightSuffix)" />
+	<Property Name="ExtraGauntletEditorTestArgs" Value="$(ExtraGauntletEditorTestArgs) -ICVFXTest.BuildName=$(BuildName)" />
 
 	<!-- Common Gauntlet settings -->
 	<Include Script="../../../../../Engine/Build/Graph/Tasks/Inc/GauntletSettings.xml" />
diff --git a/Engine/Plugins/VirtualProduction/ICVFXTesting/Source/ICVFXTesting/Private/ICVFXTestControllerAutoTest.cpp b/Engine/Plugins/VirtualProduction/ICVFXTesting/Source/ICVFXTesting/Private/ICVFXTestControllerAutoTest.cpp
--- a/Engine/Plugins/VirtualProduction/ICVFXTesting/Source/ICVFXTesting/Private/ICVFXTestControllerAutoTest.cpp
+++ b/Engine/Plugins/VirtualProduction/ICVFXTesting/Source/ICVFXTesting/Private/ICVFXTestControllerAutoTest.cpp
@@ -18,8 +18,6 @@
 #include "PlatformFeatures.h"
 #include "ProfilingDebugging/CsvProfiler.h"
 #include "VideoRecordingSystem.h"
-#include "HAL/PlatformFilemanager.h"
-#include "Misc/Paths.h"
 #include "Roles/LiveLinkTransformRole.h"
 #include "UObject/SoftObjectPtr.h"
 
@@ -181,15 +179,7 @@ public:
 		const FString TraceFile = ICVFXTest::CVarTraceFileName.GetValueOnAnyThread();
 		if (!TraceFile.IsEmpty())
 		{
-			// Delete trace file if already exists, otherwise new traces won't be stored
-			IPlatformFile& PlatformFile = FPlatformFileManager::Get().GetPlatformFile();
-			if (PlatformFile.FileExists(*TraceFile))
-			{
-				UE_LOG(LogICVFXTest, Log, TEXT("Trace file '%s' already exists. Removing it..."), *TraceFile);
-				PlatformFile.DeleteFile(*TraceFile);
-			}
-
-			Controller->ConsoleCommand(*(FString(TEXT("trace.file ")) + TraceFile));
+			Controller->ConsoleCommand(*(FString(TEXT("trace.file ")) + ICVFXTest::CVarTraceFileName.GetValueOnAnyThread()));
 		}
 
 		Controller->GoToTestLocation(0);
diff --git a/Engine/Plugins/VirtualProduction/Switchboard/Source/Switchboard/switchboard/devices/ndisplayTest/plugin_ndisplayTest.py b/Engine/Plugins/VirtualProduction/Switchboard/Source/Switchboard/switchboard/devices/ndisplayTest/plugin_ndisplayTest.py
--- a/Engine/Plugins/VirtualProduction/Switchboard/Source/Switchboard/switchboard/devices/ndisplayTest/plugin_ndisplayTest.py
+++ b/Engine/Plugins/VirtualProduction/Switchboard/Source/Switchboard/switchboard/devices/ndisplayTest/plugin_ndisplayTest.py
@@ -2,6 +2,8 @@
 
 import glob
 import os
+import threading
+import time
 import uuid
 
 from datetime import datetime
@@ -12,7 +14,7 @@ from typing import Callable
 from switchboard import message_protocol
 from switchboard import switchboard_utils as sb_utils
 from switchboard import switchboard_widgets as sb_widgets
-from switchboard.config import CONFIG, StringSetting
+from switchboard.config import CONFIG, BoolSetting, IntSetting, StringSetting, map_name_is_valid
 from switchboard.devices.device_base import DeviceStatus
 from switchboard.devices.device_widget_base import AddDeviceDialog
 from switchboard.devices.ndisplay.plugin_ndisplay import DevicenDisplay, DeviceWidgetnDisplay, LaunchMode
@@ -61,8 +63,9 @@ class AddnDisplayTestDialog(AddDeviceDialog):
 
         # Test Dir: Text field
         self.txtTestDir = QtWidgets.QLineEdit()
-        self.txtTestDir.setReadOnly(True)
+        self.txtTestDir.setReadOnly(False)
         self.txtTestDir.setText(self.test_dir)
+        self.txtTestDir.textChanged.connect(lambda text: self.set_test_location(str(text).strip(' "'), False))
 
         # Test Dir: Button to browse for a test directory
         self.btnBrowseTestDir = QtWidgets.QPushButton(self, text="Browse")
@@ -108,11 +111,12 @@ class AddnDisplayTestDialog(AddDeviceDialog):
         # Update GUI contents
         self.refresh_test_list()
 
-    def set_test_location(self, test_location):
+    def set_test_location(self, test_location, update_editbox = True):
         ''' Sets new test location and re-scans for available tests '''
         LOGGER.info(f"Set new test location: {test_location}")
         self.test_dir = test_location
-        self.txtTestDir.setText(self.test_dir)
+        if update_editbox:
+            self.txtTestDir.setText(self.test_dir)
         self.refresh_test_list()
 
     def refresh_test_list(self):
@@ -232,8 +236,11 @@ class DevicenDisplayTest(DevicenDisplay):
 
     add_device_dialog = AddnDisplayTestDialog
     current_test: nDisplayTest = None
+    last_launch_time: datetime = None
     artifacts_dir = None
-    
+    addr_run_tracking = { }
+    addr_run_tracking_lock = threading.Lock()
+
     tsettings = {
         'ndisplay_test_file': StringSetting(
             attr_name="ndisplay_test_file",
@@ -244,6 +251,54 @@ class DevicenDisplayTest(DevicenDisplay):
             is_read_only=True,
             category="General Settings",
         ),
+        'game_target_override': StringSetting(
+            attr_name="game_target_override",
+            nice_name="Override Target",
+            value="",
+            tool_tip="Game target name. By default %ProjectName% is used.",
+            allow_reset=True,
+            is_read_only=False,
+            category="Test Configuration",
+        ),
+        'editor_target_override': StringSetting(
+            attr_name="editor_target_override",
+            nice_name="Override Editor Target",
+            value="",
+            tool_tip="Editor target name. By default %ProjectName%Editor is used.",
+            allow_reset=True,
+            is_read_only=False,
+            category="Test Configuration",
+        ),
+        'test_build_info': StringSetting(
+            attr_name="test_build_info",
+            nice_name="Test Build Info",
+            value="",
+            tool_tip="Custom test information. If not set, engine&project CL will be specified.",
+            allow_reset=True,
+            is_read_only=False,
+            category="Test Configuration",
+        ),
+        'test_multi_node_delay': IntSetting(
+            attr_name="test_multi_node_delay",
+            nice_name="Multi-nodes Start Delay",
+            value=5,
+            tool_tip=("Delay between starting multiple nodes on the same address"),
+            category="Test Configuration",
+        ),
+        'test_run_boot_editor': BoolSetting(
+            attr_name='test_run_boot_editor',
+            nice_name='Run Editor Boot Test',
+            value=False,
+            tool_tip=('Whether a short editor session should run before actual test'),
+            category="Test Configuration",
+        ),
+        'report_24fps_format': BoolSetting(
+            attr_name='report_24fps_format',
+            nice_name='Generate 24 FPS Report ',
+            value=False,
+            tool_tip=('Should generate 24-fps based report (experimental)'),
+            category="Test Configuration",
+        ),
     }
 
 
@@ -267,7 +322,13 @@ class DevicenDisplayTest(DevicenDisplay):
     def plugin_settings(cls):
         ''' Returns common settings that belong to all devices of this class. '''
         return [
-            cls.tsettings['ndisplay_test_file']
+            cls.tsettings['ndisplay_test_file'],
+            cls.tsettings['game_target_override'],
+            cls.tsettings['editor_target_override'],
+            cls.tsettings['test_build_info'],
+            cls.tsettings['test_multi_node_delay'],
+            cls.tsettings['test_run_boot_editor'],
+            cls.tsettings['report_24fps_format'],
         ] + DevicenDisplay.plugin_settings()
 
     @property
@@ -293,10 +354,10 @@ class DevicenDisplayTest(DevicenDisplay):
 
     def on_open_test_artifacts(self):
         ''' Opens test artifacts location in a browser '''
-        if self.artifacts_dir.exists():
+        if self.artifacts_dir and self.artifacts_dir.exists():
             os.startfile(self.artifacts_dir)
         else:
-            LOGGER.warning('Output directory not found. Pleas try to open manually.')
+            LOGGER.warning('Output directory not found.')
             
     def get_connected_devices(self):
         ''' Returns a list with the connected devices/nodes '''
@@ -429,6 +490,25 @@ class DevicenDisplayTest(DevicenDisplay):
 
         LOGGER.info(f'Test output dir: {str(DevicenDisplayTest.artifacts_dir)}')
 
+    def pre_launch_delay(self):
+        ''' Performs conditional delay before launch '''
+
+        # Exclusive access
+        with DevicenDisplayTest.addr_run_tracking_lock:
+            # If any node has already started on this host
+            if self.address in DevicenDisplayTest.addr_run_tracking:
+                # Check how long ago
+                time_delta = datetime.now() - DevicenDisplayTest.addr_run_tracking[self.address]
+                launch_delay = self.tsettings['test_multi_node_delay'].get_value()
+                # Sleep if required
+                if(time_delta.seconds < launch_delay):
+                    sleep_time = launch_delay - time_delta.seconds
+                    LOGGER.info(f'Pre-launch delay: {sleep_time} seconds')
+                    time.sleep(sleep_time)
+
+            # Finally, update the most recent launch time for this address
+            DevicenDisplayTest.addr_run_tracking[self.address] = datetime.now()
+
     def pre_launch_check(self) -> bool:
         ''' Checks if everything is Ok before launch '''
 
@@ -462,6 +542,12 @@ class DevicenDisplayTest(DevicenDisplay):
     def launch(self, map_name):
         ''' [override] Starts an instance of device '''
 
+        # Delay launch if neccessary
+        self.pre_launch_delay()
+
+        # Update launch time
+        self.last_launch_time = datetime.now()
+
         # Check if we're ready to go
         if not self.pre_launch_check():
             LOGGER.error("Pre-launch check failed")
@@ -501,7 +587,6 @@ class DevicenDisplayTest(DevicenDisplay):
             f'-Script=\"{self.remote_path_test}\" ' + \
             f'-Target=\"BuildAndTest {project_name}\" ' + \
             f'-ScriptDir=\"{script_dir}\" ' + \
-            '-ubtargs="-NoMutex" ' + \
             f'{param_extra}'
 
         # Update last launch command
@@ -539,25 +624,76 @@ class DevicenDisplayTest(DevicenDisplay):
     def generate_extra_cmdline_params(self):
         ''' Generates additional cmdline parameters for the test session '''
 
+        # Aux lambda to convert boolean to lowercase text
+        bool_str = lambda b: "true" if b else "false"
+        # Aux lambda to convert changlist to string
+        cl_str = lambda cl: str(cl) if cl else "na"
+
         # Temporary & params
+        project_name = CONFIG.PROJECT_NAME.get_value()
+        project_path = CONFIG.get_project_dir()
         with_trace = True if CONFIG.INSIGHTS_TRACE_ENABLE.get_value() else False
-        win_pos = self.settings['window_position'].get_value()
-        win_res = self.settings['window_resolution'].get_value()
+        trace_file = self.generate_trace_file_name()
+        wnd_pos = self.settings['window_position'].get_value()
+        wnd_res = self.settings['window_resolution'].get_value()
         fullscreen = self.settings['fullscreen'].get_value()
         headless = self.settings['headless'].get_value()
+        run_editor_boot_test = self.tsettings['test_run_boot_editor'].get_value()
+        max_gpu_count = DevicenDisplay.csettings["max_gpu_count"].get_value(self.name)
 
         # Generate extra args
-        return '-set:IsOnStagePerfTest=true ' + \
-            f'-set:SkipPerfReportServer=true ' + \
+        extra_args = \
+            f'-set:ProjectName={project_name} ' + \
+            f'-set:ProjectPath="{project_path}" ' + \
+            f'-set:IsOnStagePerfTest={bool_str(True)} ' + \
+            f'-set:SkipPerfReportServer={bool_str(True)} ' + \
+            f'-set:EditorBootTest={bool_str(run_editor_boot_test)} ' + \
             f'-set:TestNameExt={self.test_name} ' + \
-            f'-set:WithTrace={with_trace} ' + \
+            f'-set:WithTrace={bool_str(with_trace)} ' + \
+            f'-set:TraceFileName={trace_file} ' + \
             f'-set:DisplayClusterNodeName={self.name} ' + \
-            f'-set:Fullscreen={fullscreen} ' + \
-            f'-set:Offscreen={headless} ' + \
-            f'-set:WndLocX={win_pos[0]} ' + \
-            f'-set:WndLocY={win_pos[1]} ' + \
-            f'-set:WndResX={win_res[0]} ' + \
-            f'-set:WndResY={win_res[1]} '
+            f'-set:Fullscreen={bool_str(fullscreen)} ' + \
+            f'-set:Offscreen={bool_str(headless)} ' + \
+            f'-set:WndLocX={wnd_pos[0]} ' + \
+            f'-set:WndLocY={wnd_pos[1]} ' + \
+            f'-set:WndResX={wnd_res[0]} ' + \
+            f'-set:WndResY={wnd_res[1]} ' + \
+            f'-set:MaxGPUCount={max_gpu_count} '
+
+        # [Optional] game target
+        target_game = self.tsettings['game_target_override'].get_value()
+        if target_game:
+            extra_args += f'-set:TargetName={target_game} '
+
+        # [Optional] editor target
+        target_edtr = self.tsettings['editor_target_override'].get_value()
+        if target_edtr:
+            extra_args += f'-set:EditorTargetName={target_edtr} '
+
+        # [Optional] explicit level
+        level = CONFIG.CURRENT_LEVEL
+        if map_name_is_valid(level):
+            extra_args += f'-set:MapOverride={level} '
+
+        # [Optional] 24fps report
+        use_24fps_report = self.tsettings['report_24fps_format'].get_value()
+        if use_24fps_report:
+            extra_args += f'-set:SummaryReportType=VP24fps '
+
+        # [OverrideDefault] Test build info
+        build_info = self.tsettings['test_build_info'].get_value()
+        if build_info:
+            # If specified, use explicit info
+            extra_args += f'-set:TestBuildName={build_info} '
+        else:
+            # Otherwise, use the default one with engine & project CL numbers
+            extra_args += f'-set:TestBuildName=E:{cl_str(self.engine_changelist)}-P:{cl_str(self.project_changelist)} '
+
+        return extra_args
+
+    def generate_trace_file_name(self):
+        ''' Generates the name of a trace file for this device. '''
+        return f'{self.name}_{self.last_launch_time.strftime("%d%m%y_%I%M%S")}.utrace'
 
     def generate_remote_test_root_directory(self):
         ''' Generates remote test root directory path. '''
@@ -583,7 +719,7 @@ class DevicenDisplayTest(DevicenDisplay):
     def generate_remote_test_utrace_filepath(self):
         ''' Generates remote test utrace file path. '''
         remote_file = self.generate_remote_test_root_directory()
-        remote_file = remote_file / f'{self.name}_Base.utrace'
+        remote_file = remote_file / self.generate_trace_file_name()
         return remote_file
 
     def generate_remote_test_report_dir(self):
@@ -599,6 +735,10 @@ class DevicenDisplayTest(DevicenDisplay):
     ):
         ''' [override] Handles test & transfer finish '''
 
+        # Reset address tracking data before next launch
+        with DevicenDisplayTest.addr_run_tracking_lock:
+            DevicenDisplayTest.addr_run_tracking.clear()
+
         # Let the base class do most of handling
         super().do_program_ended_update(
             program_name=program_name,
