Switchboard: Poll all devices sharing a workspace to influence build actions.

Ensures that Live Link Hub devices colocated with Unreal devices are able to queue their target to be built at the same time.

Simultaneously fixes a long-standing corner case with multiple local devices, ensuring the local MU server executables are always built regardless of which device is selected to build.

Also includes an unrelated fix for an asyncio issue related to potentially fulfilling a future from an incorrect/unsafe thread.

diff --git a/Engine/Plugins/VirtualProduction/Switchboard/Source/Switchboard/switchboard/devices/livelinkhub/plugin_livelinkhub.py b/Engine/Plugins/VirtualProduction/Switchboard/Source/Switchboard/switchboard/devices/livelinkhub/plugin_livelinkhub.py
--- a/Engine/Plugins/VirtualProduction/Switchboard/Source/Switchboard/switchboard/devices/livelinkhub/plugin_livelinkhub.py
+++ b/Engine/Plugins/VirtualProduction/Switchboard/Source/Switchboard/switchboard/devices/livelinkhub/plugin_livelinkhub.py
@@ -4,13 +4,15 @@ from __future__ import annotations
 
 import pathlib
 from typing import Callable
+from uuid import UUID
 
 from PySide6 import QtCore, QtWidgets
 
 from switchboard import message_protocol, switchboard_widgets
 from switchboard.config import CONFIG, Setting, BoolSetting, FilePathSetting, StringSetting
 from switchboard.devices.device_base import Device, DeviceStatus, PluginHeaderWidgets
-from switchboard.devices.unreal.plugin_unreal import DeviceUnreal, DeviceWidgetUnreal, ProgramStartQueueItem
+from switchboard.devices.unreal.plugin_unreal import (DeviceUnreal, DeviceWidgetUnreal, ProgramStartQueue,
+                                                      ProgramStartQueueItem)
 from switchboard.switchboard_logging import LOGGER
 
 
@@ -240,6 +242,15 @@ class DeviceLiveLinkHub(DeviceUnreal):
         # Not relevant to our device type.
         return
 
+    #@override
+    def _queue_all_builds(
+        self,
+        requesting_device: DeviceUnreal,
+        puuid_dependency: UUID | None = None,
+    ) -> UUID | None:
+        ubt_args = (f'LiveLinkHub {self.target_platform} Development -Progress')
+        return requesting_device._queue_build('livelinkhub', ubt_args, puuid_dependency)
+
 
 class DeviceWidgetLiveLinkHub(DeviceWidgetUnreal):
     signal_device_widget_primary = QtCore.Signal(object)
diff --git a/Engine/Plugins/VirtualProduction/Switchboard/Source/Switchboard/switchboard/devices/unreal/plugin_unreal.py b/Engine/Plugins/VirtualProduction/Switchboard/Source/Switchboard/switchboard/devices/unreal/plugin_unreal.py
--- a/Engine/Plugins/VirtualProduction/Switchboard/Source/Switchboard/switchboard/devices/unreal/plugin_unreal.py
+++ b/Engine/Plugins/VirtualProduction/Switchboard/Source/Switchboard/switchboard/devices/unreal/plugin_unreal.py
@@ -148,7 +148,7 @@ class ProgramStartQueue:
 
         # Indexed caches for the readers.
         self.rd_running_progs_by_name: dict[str, list[ProgramStartQueueItem]] = dict()
-        self.rd_puuid_by_name: dict[uuid.UUID, str] = dict()
+        self.rd_puuid_by_name: dict[str, uuid.UUID] = dict()
 
     # Write operations
 
@@ -332,7 +332,7 @@ class ProgramStartQueue:
         progs = self.running_programs_named(name)
         return [prog.puuid for prog in progs]
 
-    def puuid_from_name(self, name) -> uuid.UUID:
+    def puuid_from_name(self, name: str) -> uuid.UUID:
         '''
         Returns the puuid of the specified program name.
 
@@ -1862,21 +1862,22 @@ class DeviceUnreal(Device):
         if self.exclude_from_build.get_value():
             return
 
+        # This should only be possible when invoked on individual devices
+        # manually; SwitchboardDialog._sync_build_all already only invokes
+        # this method for a single device per workspace.
         for device in self.devices_sharing_workspace():
-            if device.status >= DeviceStatus.OPEN:
+            if device.status >= DeviceStatus.BUILDING:
                 LOGGER.warning(f'{self.name}: Unable to build, because '
                                f'another device ({device.name}) using the '
-                               'same workspace is already running.')
+                               'same workspace is already running/building.')
                 return
 
         program_name = UnrealJobs.BuildProject.value
 
         # check if it is already on its way:
         try:
-            existing_puuid = self.program_start_queue.puuid_from_name(
-                program_name)
-            LOGGER.info(
-                f"{self.name}: Already building with puuid {existing_puuid}")
+            existing_puuid = self.program_start_queue.puuid_from_name(program_name)
+            LOGGER.info(f"{self.name}: Already building with puuid {existing_puuid}")
             return
         except KeyError:
             pass
@@ -1890,15 +1891,30 @@ class DeviceUnreal(Device):
             pass
 
         if sync_puuid:
-            LOGGER.debug(
-                f"{self.name} Queuing build after sync with "
-                f"puuid {sync_puuid}")
+            LOGGER.debug(f"{self.name} Queuing build after sync with puuid {sync_puuid}")
 
         # Build dependency chain
         puuid_dependency = sync_puuid
 
-        # TODO: Corner case if multiple local devices, and we build on a single
-        # local device other than chosen. Guarantee on any local single build?
+        # We've picked one device to act as build representative for our workspace;
+        # however, we want to let the others weigh in, especially derived classes.
+        for device in self.devices_sharing_workspace(only_connected=False):
+            puuid_dependency = device._queue_all_builds(self, puuid_dependency)
+
+    def _queue_all_builds(
+        self,
+        requesting_device: DeviceUnreal,
+        puuid_dependency: uuid.UUID | None = None,
+    ) -> uuid.UUID | None:
+        '''
+        Customization point for derived classes to customize build actions.
+
+        The passed-in `requesting_device` may be a different device, sharing our workspace, to use as a proxy.
+
+        Adding "redundant" actions should be safe; `_queue_build` (and subsequently the various `_build_*` wrapper
+        methods) will early out when attempting to queue a second build with the same name.
+        '''
+
         if (self.is_designated_local_builder() and
                 CONFIG.ENGINE_SYNC_METHOD.get_value() == EngineSyncMethod.Build_Engine.value):
             # Build multi-user server
@@ -1914,24 +1930,20 @@ class DeviceUnreal(Device):
                     if mb_ret == QtWidgets.QMessageBox.Yes:
                         DeviceUnreal.mu_server.process.kill()
 
-                puuid_dependency = self._build_mu_server(
-                    puuid_dependency=puuid_dependency)
-
-                puuid_dependency = self._build_mu_slate_server(
-                    puuid_dependency=puuid_dependency)
+                puuid_dependency = requesting_device._build_mu_server(puuid_dependency)
+                puuid_dependency = requesting_device._build_mu_slate_server(puuid_dependency)
 
             # TODO: Build listener
             # - Handle the running listener's exe being locked
             # - Parse SwitchboardListenerVersion.h, skip if unchanged?
 
             # if CONFIG.LISTENER_AUTO_BUILD:
-            #     puuid_dependency = self._build_listener(
-            #         puuid_dependency=puuid_dependency)
+            #     puuid_dependency = requesting_device._build_listener(puuid_dependency)
 
-        puuid_dependency = self._build_shadercompileworker(
-            puuid_dependency=puuid_dependency)
-        puuid_dependency = self._build_project(
-            puuid_dependency=puuid_dependency)
+        puuid_dependency = requesting_device._build_shadercompileworker(puuid_dependency)
+        puuid_dependency = requesting_device._build_project(puuid_dependency)
+
+        return puuid_dependency
 
     @property
     def target_platform(self) -> str:
@@ -1939,44 +1951,52 @@ class DeviceUnreal(Device):
         # FIXME?: Not strictly correct, but holds for desktop host platforms
         return self.platform_binary_directory
 
-    def _build_project(self, puuid_dependency: Optional[uuid.UUID] = None):
+    def _build_project(self, puuid_dependency: Optional[uuid.UUID] = None) -> uuid.UUID | None:
         ubt_args = (
             f'{self.target_platform} Development '
             f'-project="{CONFIG.UPROJECT_PATH.get_value(self.name)}" '
             '-TargetType=Editor -Progress -NoHotReloadFromIDE')
-        return self._queue_build(
-            'project', ubt_args=ubt_args, puuid_dependency=puuid_dependency)
-
-    def _build_mu_server(self, puuid_dependency: Optional[uuid.UUID] = None):
-        ubt_args = (f'UnrealMultiUserServer {self.target_platform} '
-                    'Development -Progress')
-        return self._queue_build(
-            'mu_server', ubt_args=ubt_args, puuid_dependency=puuid_dependency)
-
-    def _build_mu_slate_server(self, puuid_dependency: Optional[uuid.UUID] = None):
-        ubt_args = (f'UnrealMultiUserSlateServer {self.target_platform} '
-                    'Development -Progress')
-        return self._queue_build(
-            'mu_slate_server', ubt_args=ubt_args, puuid_dependency=puuid_dependency)
-    
-    def _build_listener(self, puuid_dependency: Optional[uuid.UUID] = None):
-        ubt_args = (f'SwitchboardListener {self.target_platform} '
-                    'Development -Progress')
-        return self._queue_build(
-            'listener', ubt_args=ubt_args, puuid_dependency=puuid_dependency)
-
-    def _build_shadercompileworker(
-            self, puuid_dependency: Optional[uuid.UUID] = None):
-        ubt_args = (f'ShaderCompileWorker {self.target_platform} '
-                    'Development -Progress')
-        return self._queue_build(
-            'shadercw', ubt_args=ubt_args, puuid_dependency=puuid_dependency)
+        return self._queue_build('project', ubt_args, puuid_dependency)
+
+    def _build_mu_server(self, puuid_dependency: Optional[uuid.UUID] = None) -> uuid.UUID | None:
+        ubt_args = f'UnrealMultiUserServer {self.target_platform} Development -Progress'
+        return self._queue_build('mu_server', ubt_args, puuid_dependency)
+
+    def _build_mu_slate_server(self, puuid_dependency: Optional[uuid.UUID] = None) -> uuid.UUID | None:
+        ubt_args = f'UnrealMultiUserSlateServer {self.target_platform} Development -Progress'
+        return self._queue_build('mu_slate_server', ubt_args, puuid_dependency)
+
+    def _build_listener(self, puuid_dependency: Optional[uuid.UUID] = None) -> uuid.UUID | None:
+        ubt_args = f'SwitchboardListener {self.target_platform} Development -Progress'
+        return self._queue_build('listener', ubt_args, puuid_dependency)
+
+    def _build_shadercompileworker(self, puuid_dependency: Optional[uuid.UUID] = None) -> uuid.UUID | None:
+        ubt_args = f'ShaderCompileWorker {self.target_platform} Development -Progress'
+        return self._queue_build('shadercw', ubt_args, puuid_dependency)
 
     def _queue_build(
-            self, program_name_suffix: str, ubt_args: str,
-            puuid_dependency: Optional[uuid.UUID] = None):
+        self,
+        program_name_suffix: str,
+        ubt_args: str,
+        puuid_dependency: Optional[uuid.UUID] = None,
+    ) -> uuid.UUID | None:
+        '''
+        Adds a UBT invocation to the queue as "build_{program_name_suffix}",
+        but only if an entry with that name does not already exist.
+
+        If no new build was queued, returns the `puuid_dependency` which was
+        passed in, in order to make chaining calls more convenient.
+        '''
+
         program_name = f"{UnrealJobs.BuildPrefix.value}{program_name_suffix}"
 
+        try:
+            existing_puuid = self.program_start_queue.puuid_from_name(program_name)
+            LOGGER.debug(f'{self.name}: Not queuing {program_name}; exists as {existing_puuid}')
+            return puuid_dependency
+        except KeyError:
+            pass
+
         engine_path = CONFIG.ENGINE_DIR.get_value(self.name)
         if self.target_platform.lower().startswith('win'):
             ubt_path = os.path.join(engine_path, 'Build', 'BatchFiles',
@@ -1998,7 +2018,7 @@ class DeviceUnreal(Device):
             LOGGER.info(
                 f"{self.name}: Sending {program_name} command: "
                 f"{ubt_path} {ubt_args}")
-            
+
             for device in self.devices_sharing_workspace():
                 device.status = DeviceStatus.BUILDING
 
diff --git a/Engine/Plugins/VirtualProduction/Switchboard/Source/Switchboard/switchboard/listener_client.py b/Engine/Plugins/VirtualProduction/Switchboard/Source/Switchboard/switchboard/listener_client.py
--- a/Engine/Plugins/VirtualProduction/Switchboard/Source/Switchboard/switchboard/listener_client.py
+++ b/Engine/Plugins/VirtualProduction/Switchboard/Source/Switchboard/switchboard/listener_client.py
@@ -168,11 +168,14 @@ class ListenerQtHandler(QtCore.QObject):
                 f'(Fingerprint: {self.fingerprint_desc})',
                 QtWidgets.QLineEdit.EchoMode.Password)
 
-            assert self.future
             if blob and ok:
-                self.future.set_result(blob)
+                result = blob
             else:
-                self.future.set_result('')
+                result = ''
+
+            assert self.future
+            future = self.future
+            future.get_loop().call_soon_threadsafe(lambda: future.set_result(result))
 
 
     class PromptFingerprintMismatch(Prompt):
diff --git a/Engine/Plugins/VirtualProduction/Switchboard/Source/Switchboard/switchboard/switchboard_dialog.py b/Engine/Plugins/VirtualProduction/Switchboard/Source/Switchboard/switchboard/switchboard_dialog.py
--- a/Engine/Plugins/VirtualProduction/Switchboard/Source/Switchboard/switchboard/switchboard_dialog.py
+++ b/Engine/Plugins/VirtualProduction/Switchboard/Source/Switchboard/switchboard/switchboard_dialog.py
@@ -1320,15 +1320,20 @@ class SwitchboardDialog(QtCore.QObject):
                 LOGGER.info(f'Skipping {actions_str} for {device.name} because it has empty workspaces')
                 continue
 
+            already_queued = False
             for workspace in workspaces:
                 if existing_device := workspace_to_device_map.get(workspace):
                     LOGGER.info(f'Skipping {actions_str} for {device.name} '
                                 f'because {existing_device.name} is already '
                                 'queued using the same workspace')
+                    already_queued = True
                     continue
 
                 workspace_to_device_map[workspace] = device
 
+            if already_queued:
+                continue
+
             if do_sync:
                 device.widget.sync_button_clicked()
 
