nDisplay: Fix PSO rebuild hitches in -game runs.

Increased default PSO eviction time in nDisplay launches to 3600 seconds. It can still be overridden from Switchboard if desired.

Note: PSO pre-caching is currently only supported in non-WITH_EDITOR builds.

#rnx

diff --git a/Engine/Plugins/VirtualProduction/Switchboard/Source/Switchboard/switchboard/devices/ndisplay/plugin_ndisplay.py b/Engine/Plugins/VirtualProduction/Switchboard/Source/Switchboard/switchboard/devices/ndisplay/plugin_ndisplay.py
--- a/Engine/Plugins/VirtualProduction/Switchboard/Source/Switchboard/switchboard/devices/ndisplay/plugin_ndisplay.py
+++ b/Engine/Plugins/VirtualProduction/Switchboard/Source/Switchboard/switchboard/devices/ndisplay/plugin_ndisplay.py
@@ -1304,6 +1304,13 @@ class DevicenDisplay(DeviceUnreal):
         if DevicenDisplay.csettings['incremental_gc_reachability'].get_value(self.name):
             dp_cvars.append('gc.AllowIncrementalReachability=1')
 
+        # PSO Eviction time.
+        #
+        # PSO eviction may cause PSO rebuilding and may induce hitches.
+        # Currently, only non-editor builds support pre-caching, so for now a way to avoid
+        # these hitches is by lengthening the eviction time from the default 60 seconds.
+        dp_cvars.append(f'r.pso.evictiontime={60*60}')
+
         # Add user set dp cvars, overriding any of the forced ones.
         user_dp_cvars = self.conform_dpcvars(self.csettings['ndisplay_dp_cvars'].get_value(self.name))
         dp_cvars = self.add_or_override_cvars(dp_cvars, user_dp_cvars)
