Switchboard: Fix Live Link Hub-related device build regressions (related to CL 42455953).

- `UnrealDevice.target_platform` property is only valid to access on connected devices, but `_queue_all_builds` can be invoked on disconnected devices, so use the `requesting_device` to retrieve `target_platform` instead.

- `build_project` is no longer guaranteed to be the final build in the queue, so modify `do_program_ended_update` to look for the absence of remaining `build_*` jobs instead.

#rnx

diff --git a/Engine/Plugins/VirtualProduction/Switchboard/Source/Switchboard/switchboard/devices/livelinkhub/plugin_livelinkhub.py b/Engine/Plugins/VirtualProduction/Switchboard/Source/Switchboard/switchboard/devices/livelinkhub/plugin_livelinkhub.py
--- a/Engine/Plugins/VirtualProduction/Switchboard/Source/Switchboard/switchboard/devices/livelinkhub/plugin_livelinkhub.py
+++ b/Engine/Plugins/VirtualProduction/Switchboard/Source/Switchboard/switchboard/devices/livelinkhub/plugin_livelinkhub.py
@@ -248,7 +248,7 @@ class DeviceLiveLinkHub(DeviceUnreal):
         requesting_device: DeviceUnreal,
         puuid_dependency: UUID | None = None,
     ) -> UUID | None:
-        ubt_args = (f'LiveLinkHub {self.target_platform} Development -Progress')
+        ubt_args = (f'LiveLinkHub {requesting_device.target_platform} Development -Progress')
         return requesting_device._queue_build('livelinkhub', ubt_args, puuid_dependency)
 
 
diff --git a/Engine/Plugins/VirtualProduction/Switchboard/Source/Switchboard/switchboard/devices/unreal/plugin_unreal.py b/Engine/Plugins/VirtualProduction/Switchboard/Source/Switchboard/switchboard/devices/unreal/plugin_unreal.py
--- a/Engine/Plugins/VirtualProduction/Switchboard/Source/Switchboard/switchboard/devices/unreal/plugin_unreal.py
+++ b/Engine/Plugins/VirtualProduction/Switchboard/Source/Switchboard/switchboard/devices/unreal/plugin_unreal.py
@@ -2688,7 +2688,10 @@ class DeviceUnreal(Device):
                     if error_pattern.search(line):
                         LOGGER.error(f"{self.name}: {line}")
 
-            if UnrealJobs.BuildProject.value == program_name:
+            remaining_build_names = list(filter(lambda name: name.startswith(UnrealJobs.BuildPrefix.value),
+                                                self.program_start_queue.rd_puuid_by_name.keys()))
+
+            if not remaining_build_names:
                 for device in self.devices_sharing_workspace():
                     device.status = DeviceStatus.CLOSED
 
