Switchboard: sbl_helper.py now explicitly maintains workspace sync state in files in the engine/project directories.

Previously, the indication of what changelist(s) a device is currently synced to were based on invoking `p4 cstat`, passing engine or project directory, followed by a `...#have` wildcard + revision specifier. However, this approach left a number of edge cases unhandled, such as failing to reflect subsets of files synced to older revisions, changelists consisting entirely of file deletions, or changelists only affecting "external" paths referenced via `AdditionalPathsToSync`.

`sbl_helper.py` now implements a `syncstatus` command which is used instead. During each `sync` command, we write a file (`.sblsync.json`) to both the engine and project directories, containing a record of the corresponding changelist the user had selected for sync, as well as a timestamp. Then, during subsequent `syncstatus` commands, we can simply "preview" a sync to that changelist, and rely on Perforce to tell us whether it would have taken any action to update the workspace (i.e. whether it's "dirty" and actually needs to be synced). If the workspace is "dirty", this will be indicated by an asterisk next to the changelist indicators in Switchboard.

Other changes:
- p4 calls now use an appropriate working directory, in order to respect P4CONFIG (if present).
- All p4_utils helpers now enable P4PORT to be passed in explicitly and propagate the provided value.
- Avoids redundant UnrealGameSync.ini P4 queries for DeviceUnreal subclasses.
- Type annotations in the Device base class now correctly reflect that the changelist attributes are frequently assigned and used as strings (previously incidental, now relied upon).
- Removed a misleadingly-named passthrough "init" method on the `Config` class.
- Lays the groundwork for additional automation of first-time workspace/config setup.

diff --git a/Engine/Plugins/VirtualProduction/Switchboard/Source/Switchboard/docs/SyncAndVersioning.md b/Engine/Plugins/VirtualProduction/Switchboard/Source/Switchboard/docs/SyncAndVersioning.md
--- a/Engine/Plugins/VirtualProduction/Switchboard/Source/Switchboard/docs/SyncAndVersioning.md
+++ b/Engine/Plugins/VirtualProduction/Switchboard/Source/Switchboard/docs/SyncAndVersioning.md
@@ -120,7 +120,11 @@ Because of this, some versioning concepts designed with UnrealGameSync in mind d
 Switchboard has per-device indicators of sync and build status, driven by various aspects of the device's workspace.
 
 ### Sync
-The indication of what changelists a device is synced to is based on invoking `p4 cstat` for the engine or project paths with the `#have` revision specifier. The changelist returned by the command is compared against the one selected from the corresponding Switchboard combo box. If there's a mismatch, the corresponding changelist indicator and the sync button will both be highlighted, indicating the device needs to be synced to the desired changelist.
+In older versions of Switchboard, the indication of what changelist(s) a device is currently synced to were based on invoking `p4 cstat`, passing engine or project directory, followed by a `...#have` wildcard + revision specifier. However, this approach left a number of edge cases unhandled, such as failing to reflect subsets of files synced to older revisions, changelists consisting entirely of file deletions, or changelists only affecting "external" paths referenced via `AdditionalPathsToSync`.
+
+As of UE 5.7, `sbl_helper.py` now implements a `syncstatus` command which is used instead. During each `sbl_helper.py sync` command, we write a file (`.sblsync.json`) to both the engine and project directories, containing a record of the corresponding changelist the user had selected for sync, as well as a timestamp. Then, during subsequent `syncstatus` commands, we can simply "preview" a sync to that changelist, and rely on Perforce to tell us whether it would have taken any action to update the workspace (i.e. whether it's "dirty" and actually needs to be synced). If the workspace is "dirty", this will be indicated by an asterisk next to the changelist indicators in Switchboard.
+
+The changelist returned by the command is compared against the one selected from the corresponding Switchboard combo box. If there's a mismatch, the corresponding changelist indicator and the sync button will both be highlighted, indicating the device needs to be synced to the desired changelist.
 
 ### Build
 Similarly, Switchboard will indicate if a device has not yet built the changelist it's currently synced to. For this, Switchboard relies on a few relevant facts:
diff --git a/Engine/Plugins/VirtualProduction/Switchboard/Source/Switchboard/sbl_helper.py b/Engine/Plugins/VirtualProduction/Switchboard/Source/Switchboard/sbl_helper.py
--- a/Engine/Plugins/VirtualProduction/Switchboard/Source/Switchboard/sbl_helper.py
+++ b/Engine/Plugins/VirtualProduction/Switchboard/Source/Switchboard/sbl_helper.py
@@ -5,6 +5,8 @@ from __future__ import annotations
 import argparse
 import copy
 from dataclasses import dataclass
+import datetime
+from enum import StrEnum
 import io
 import json
 import logging
@@ -17,7 +19,7 @@ import sys
 import threading
 import tempfile
 import traceback
-from typing import Callable, Optional, Sequence, Union
+from typing import Any, Callable, Optional, Sequence, Union
 import uuid
 
 from switchboard import p4_utils, ugs_utils as UGS
@@ -26,7 +28,12 @@ from switchboard import p4_utils, ugs_utils as UGS
 MAX_PARALLEL_WORKERS = 8  # Consistent with UGS, ushell, etc.
 
 SCRIPT_NAME = os.path.basename(__file__)
-__version__ = '1.2.0'
+__version__ = '1.3.0'
+
+
+class Subcommands(StrEnum):
+    Sync = 'sync'
+    SyncStatus = 'syncstatus'
 
 
 class SyncFile:
@@ -99,6 +106,7 @@ def p4_sync_worker(
     clobber_queue: Optional[Queue[SyncFile]] = None,
     dry_run: bool = False,
     force: bool = False,
+    port: Optional[str] = None,
     user: Optional[str] = None,
     client: Optional[str] = None,
 ):
@@ -131,7 +139,9 @@ def p4_sync_worker(
         argfile.write(f'{item.spec}\n')
     argfile.close()  # Otherwise the p4 process can't access the file.
 
-    proc = p4_utils.p4('sync', opts, gopts, user=user, client=client)
+    p4kwargs: dict[str, Any] = {'port': port, 'user': user, 'client': client}
+
+    proc = p4_utils.p4('sync', opts, gopts, **p4kwargs)
 
     remaining_work = copy.deepcopy(sync_workload)
 
@@ -196,14 +206,17 @@ def p4_sync_worker(
 def p4_get_preview_sync_files(
     input_specs: list[str],
     *,
+    port: Optional[str] = None,
     user: Optional[str] = None,
     client: Optional[str] = None,
 ) -> SyncWorkload:
     ''' Runs a preview sync to generate a list of files and revisions. '''
     result_workload = SyncWorkload()
+
+    p4kwargs: dict[str, Any] = {'port': port, 'user': user, 'client': client}
     records = p4_utils.p4_get_records(
         'sync', ['-n', *input_specs], include_info=True, include_error=True,
-        user=user, client=client)
+        **p4kwargs)
 
     for record in records:
         if record[b'code'] == b'stat':
@@ -253,36 +266,25 @@ def friendly_bytes(size: float) -> str:
 
 
 def p4_sync(
-    input_specs: list[str],
+    total_work: SyncWorkload,
     num_workers: int = MAX_PARALLEL_WORKERS,
     *,
-    sync_filter_fn: Optional[Callable[[SyncWorkload], None]] = None,
     clobber_filter_fn: Optional[Callable[[SyncWorkload], None]] = None,
     progress_fn: Optional[Callable[[int, int, int, int], None]] = None,
     dry_run: bool = False,
+    port: Optional[str] = None,
     user: Optional[str] = None,
     client: Optional[str] = None,
-    write_workload_file: Optional[io.TextIOWrapper] = None
 ) -> Optional[int]:
     # First, gather complete list of files/revisions that need to be synced.
     logging.info('Determining files to be synced')
-    total_work = p4_get_preview_sync_files(input_specs, user=user,
-                                           client=client)
 
-    if sync_filter_fn:
-        sync_filter_fn(total_work)
+    p4kwargs: dict[str, Any] = {'port': port, 'user': user, 'client': client}
 
     if len(total_work.sync_files) < 1:
         logging.info('Nothing to do')
         return None
 
-    # User may have requested to write the pending sync files and revisions to
-    # file and exit early.
-    if write_workload_file is not None:
-        for item in total_work.sync_files.values():
-            write_workload_file.write(f'{item.spec}\n')
-        return 0
-
     # Separate deletes and run them on their own worker process. We no longer
     # refer to `total_work` after this.
     # TODO?: This is only necessary if the server is < 2020.1, which added
@@ -295,8 +297,6 @@ def p4_sync(
         else:
             nondelete_work.add(item)
 
-    del total_work
-
     # First pass: Attempt to run all deletes (in the main thread).
     total_delete_files = len(delete_work.sync_files)
     if total_delete_files > 0:
@@ -306,8 +306,7 @@ def p4_sync(
         if clobber_filter_fn:
             clobber_queue = Queue()
 
-        p4_sync_worker(delete_work, clobber_queue=clobber_queue,
-                       dry_run=dry_run, user=user, client=client)
+        p4_sync_worker(delete_work, clobber_queue=clobber_queue, dry_run=dry_run, **p4kwargs)
 
         # First pass clobber retry: Force sync any deletes where we want to
         # override noclobber.
@@ -330,8 +329,7 @@ def p4_sync(
             num_files_to_clobber = len(clobber_work.sync_files)
             if num_files_to_clobber:
                 logging.info(f'Clobbering {num_files_to_clobber} files')
-                p4_sync_worker(clobber_work, force=True, dry_run=dry_run,
-                               user=user, client=client)
+                p4_sync_worker(clobber_work, force=True, dry_run=dry_run, **p4kwargs)
 
         logging.info('Done removing files from workspace')
 
@@ -347,8 +345,7 @@ def p4_sync(
         dry_run=dry_run,
         clobber_queue=clobber_queue,
         progress_fn=progress_fn,
-        user=user,
-        client=client)
+        **p4kwargs)
 
     # Second pass clobber retry: Force sync where we want to override
     # noclobber.
@@ -377,8 +374,7 @@ def p4_sync(
                 dry_run=dry_run,
                 force=True,
                 progress_fn=progress_fn,
-                user=user,
-                client=client)
+                **p4kwargs)
 
     logging.info('Sync complete')
     return 0
@@ -392,6 +388,7 @@ def dispatch_sync_workers(
     force: bool = False,
     clobber_queue: Optional[Queue[SyncFile]] = None,
     progress_fn: Optional[Callable[[int, int, int, int], None]] = None,
+    port: Optional[str] = None,
     user: Optional[str] = None,
     client: Optional[str] = None,
 ) -> int:
@@ -437,7 +434,7 @@ def dispatch_sync_workers(
             args=(load,),
             kwargs={'completion_queue': completion_queue,
                     'clobber_queue': clobber_queue, 'force': force,
-                    'dry_run': dry_run, 'user': user, 'client': client}
+                    'dry_run': dry_run, 'port': port, 'user': user, 'client': client}
         )
         workers.append(worker)
         worker.start()
@@ -487,6 +484,7 @@ def path_is_relative_to(
 class SbListenerHelper:
     def __init__(self):
         self.parser = self.build_parser()
+        self.log_string = io.StringIO()
 
         self.uproj_path: Optional[pathlib.Path] = None
         self.engine_dir: Optional[pathlib.Path] = None
@@ -499,8 +497,11 @@ class SbListenerHelper:
         self.ugs_versioning: bool = False
 
         self.dry_run: bool = False
+        self.p4port: Optional[str] = None
         self.p4user: Optional[str] = None
         self.p4client: Optional[str] = None
+        self.p4client_project: Optional[str] = None
+        self.p4kwargs: dict[str, Any] = {}
         self.p4clientspec: dict[str, str] = {}
 
         if UGS.SyncFilters.supported():
@@ -508,121 +509,152 @@ class SbListenerHelper:
         else:
             self.ugs_filters = None
 
-    @staticmethod
-    def build_parser():
+    @classmethod
+    def build_parser(cls):
         parser = argparse.ArgumentParser()
 
         # Global options
-        parser.add_argument(
-            '--log-level', default='INFO',
-            help='DEBUG, INFO (default), WARNING, ERROR, CRITICAL')
+        parser.add_argument('--log-level', default='INFO',
+                            help='Sets the minimum severity of log messages before appearing on stderr.',
+                            choices=list(filter(lambda n: n != 'NOTSET',
+                                         [n for n, _ in logging.getLevelNamesMapping().items()] + ['NONE'])))
 
         subparsers = parser.add_subparsers(dest='action')
         subparsers.required = True
 
-        # Sync action
-        sync_parser = subparsers.add_parser('sync')
-
-        sync_parser.add_argument(
-            '--project', type=pathlib.Path, metavar='UPROJECT',
-            help='Path to .uproject file')
-        sync_parser.add_argument(
-            '--engine-dir', type=pathlib.Path,
-            help='Path to Engine directory. If omitted, the ancestors of the '
-                 'project directory will be searched')
-        sync_parser.add_argument(
-            '--engine-cl', type=int,
-            help='Changelist to sync engine to (requires UPROJECT)')
-        sync_parser.add_argument(
-            '--project-cl', type=int,
-            help='Changelist to sync project to (requires UPROJECT)')
-        sync_parser.add_argument(
-            '--generate', action='store_true',
-            help='Run GenerateProjectFiles after syncing')
-        sync_parser.add_argument(
-            '--clobber-engine', action='store_true',
-            help='Override noclobber for engine files')
-        sync_parser.add_argument(
-            '--clobber-project', action='store_true',
-            help='Override noclobber for project files')
-        sync_parser.add_argument(
-            '--ugs-versioning', action='store_true',
-            help='Maintain depot CompatibleChangelist (e.g. hotfix stream)')
+        # Parent parsers used for common options
+        p4_parser = argparse.ArgumentParser(add_help=False)
+        p4_parser.add_argument('-p', '--p4port', type=str,
+                               help='Override Perforce P4UPORT')
+        p4_parser.add_argument('-u', '--p4user', type=str,
+                               help='Override Perforce P4USER')
+        p4_parser.add_argument('-c', '--p4client', type=str,
+                               help='Override Perforce P4CLIENT')
+        p4_parser.add_argument('--p4client-project', type=str,
+                               help='Override Perforce P4CLIENT just for project operations')
+        p4_parser.add_argument('--project', type=pathlib.Path, metavar='UPROJECT',
+                               help='Path to .uproject file')
+        p4_parser.add_argument('--engine-dir', type=pathlib.Path,
+                               help='Path to Engine directory. If omitted, the ancestors of the '
+                                    'project directory will be searched')
+        p4_parser.add_argument('--engine-cl', type=int,
+                               help='Changelist to sync engine to (requires UPROJECT)')
+        p4_parser.add_argument('--project-cl', type=int,
+                               help='Changelist to sync project to (requires UPROJECT)')
+
+        # UGS sync filter options
+        p4_parser.add_argument('--include-categories', type=str,
+                                 help='List of UGS SyncCategory UUIDs to include in sync')
+        p4_parser.add_argument('--custom-view', type=str,
+                                 help='Comma-separated Perforce wildcards to exclude from sync')
+        p4_parser.add_argument('--remove-excluded-workspace-files', action='store_true',
+                                 help='Scan workspace to remove all files excluded by sync filter')
+
+        # Add subcommands
+        cls.build_sync_parser(subparsers.add_parser(Subcommands.Sync, parents=[p4_parser]))
+        cls.build_syncstatus_parser(subparsers.add_parser(Subcommands.SyncStatus, parents=[p4_parser]))
+
+        return parser
+
+    @classmethod
+    def build_sync_parser(cls, sync_parser: argparse.ArgumentParser):
+        ''' Configures subparser arguments specific to the `sync` subcommand. '''
+
+        sync_parser.add_argument('--generate', action='store_true',
+                                 help='Run GenerateProjectFiles after syncing')
+        sync_parser.add_argument('--clobber-engine', action='store_true',
+                                 help='Override noclobber for engine files')
+        sync_parser.add_argument('--clobber-project', action='store_true',
+                                 help='Override noclobber for project files')
+        sync_parser.add_argument('--ugs-versioning', action='store_true',
+                                 help='Maintain depot CompatibleChangelist (e.g. hotfix stream)')
 
         sync_parser.add_argument('-n', '--dry-run', action='store_true',
                                  help='Preview sync; do not update files')
-        sync_parser.add_argument('-u', '--p4user', type=str,
-                                 help='Override Perforce P4USER')
-        sync_parser.add_argument('-c', '--p4client', type=str,
-                                 help='Override Perforce P4CLIENT')
-
-        # UGS sync filter options
-        sync_parser.add_argument(
-            '--include-categories', type=str,
-            help='List of UGS SyncCategory UUIDs to include in sync')
-        sync_parser.add_argument(
-            '--custom-view', type=str,
-            help='Comma-separated Perforce wildcards to exclude from sync')
-        sync_parser.add_argument(
-            '--remove-excluded-workspace-files', action='store_true',
-            help='Scan workspace to remove all files excluded by sync filter')
 
         # UGS CLI integration options
-        sync_parser.add_argument(
-            '--use-ugs', action='store_true',
-            help='Specifies that UnrealGameSync should be used to sync the '
-                 'project and engine together.')
-        sync_parser.add_argument(
-            '--use-pcbs', action='store_true',
-            help='Specifies that precompiled binaries should be synced along '
-                 'with the project (requires syncing via UnrealGameSync).')
-        sync_parser.add_argument(
-            '--ugs-lib-dir', type=pathlib.Path,
-            help='Directory path specifying where to find the UnrealGameSync '
-                 'library (ugs.dll).')
+        sync_parser.add_argument('--use-ugs', action='store_true',
+                                 help='Specifies that UnrealGameSync should be used to sync the '
+                                      'project and engine together.')
+        sync_parser.add_argument('--use-pcbs', action='store_true',
+                                 help='Specifies that precompiled binaries should be synced along '
+                                      'with the project (requires syncing via UnrealGameSync).')
+        sync_parser.add_argument('--ugs-lib-dir', type=pathlib.Path,
+                                 help='Directory path specifying where to find the UnrealGameSync '
+                                      'library (ugs.dll).')
 
         # Debugging options
-        sync_parser.add_argument(
-            '--dump-sync', type=argparse.FileType('wt'),
-            metavar='SYNC_WORKLOAD_TEXT_FILE',
-            help='Write the complete list of depot paths and revisions that '
-                 'would be synced, but exit without syncing')
+        sync_parser.add_argument('--dump-sync', type=argparse.FileType('wt'), metavar='SYNC_WORKLOAD_TEXT_FILE',
+                                 help='Write the complete list of depot paths and revisions that '
+                                 'would be synced, but exit without syncing')
 
-        return parser
+    @classmethod
+    def build_syncstatus_parser(cls, syncstatus_parser: argparse.ArgumentParser):
+        ''' Configures subparser arguments specific to the `syncstatus` subcommand. '''
+        pass
 
     def run(self, args: Optional[Sequence[str]] = None) -> int:
         options = self.parser.parse_args(args)
 
+        file_handler = logging.FileHandler('sbl_helper.log', mode='w', encoding='utf-8')
+        handlers: list[logging.Handler] = [file_handler]
+        string_handler: logging.StreamHandler | None = None
+        stderr_handler: logging.StreamHandler | None = None
+
+        # file_handler is always full verbosity; log_level only affects stderr/log_string
+        if options.log_level.lower() != 'none':
+            stderr_handler = logging.StreamHandler()
+            stderr_handler.setLevel(options.log_level)
+            stderr_handler.setFormatter(logging.Formatter('%(relativeCreated)06dms %(levelname)-8s %(message)s'))
+            handlers.append(stderr_handler)
+
+            string_handler = logging.StreamHandler(self.log_string)
+            string_handler.setLevel(options.log_level)
+            handlers.append(string_handler)
+
         logging.basicConfig(
-            level=options.log_level,
-            format='%(relativeCreated)06dms %(levelname)-8s %(message)s'
+            level=logging.NOTSET,
+            format='%(asctime)s [%(levelname)-8s] <%(name)s>: %(message)s',
+            handlers=handlers
         )
 
-        if options.action == 'sync':
+        if options.action == Subcommands.Sync:
             return self.run_sync(options)
+        elif options.action == Subcommands.SyncStatus:
+            return self.run_syncstatus(options)
 
         assert False
 
-    def run_sync(self, options: argparse.Namespace) -> int:
+    def _get_sync_workload(
+        self,
+        options: argparse.Namespace,
+        *,
+        engine_dir: pathlib.Path | None = None,
+        sync_engine_cl: Optional[int] = None,
+        project_dir: pathlib.Path | None = None,
+        sync_project_cl: Optional[int] = None,
+        port: Optional[str] = None,
+        user: Optional[str] = None,
+        client: Optional[str] = None,
+    ) -> tuple[SyncWorkload, BuildVerInfo]:
+        assert engine_dir or project_dir
+
+        p4kwargs: dict[str, Any] = {'port': port, 'user': user, 'client': client}
+
         engine_paths: list[str] = []
         project_paths: list[str] = []
         pathspecs: list[str] = []
 
-        self.check_sync_options(options)
-
-        ugs_config = UGS.parse_depot_ugs_configs(
-            engine_dir=self.engine_dir, project_dir=self.uproj_path.parent,
-            p4user=self.p4user, p4client=self.p4client)
+        ugs_config = UGS.parse_depot_ugs_configs(engine_dir=engine_dir,
+                                                 project_dir=project_dir,
+                                                 **p4kwargs)
 
-        self.additional_paths_to_sync = ugs_config.try_get(
-            'Perforce', 'AdditionalPathsToSync') or []
+        self.additional_paths_to_sync = ugs_config.try_get('Perforce', 'AdditionalPathsToSync') or []
 
         # Parse sync filter categories from INI and match to command line args
-        any_active_filters = False
         if options.include_categories or options.custom_view:
             if not self.ugs_filters:
-                self.parser.error('Sync filters were specified, but filtering '
-                                  'is unavailable (missing P4Python?)')
+                self.parser.error('Sync filters were specified, but filtering is unavailable (missing P4Python?)')
 
             if options.include_categories:
                 include_ids = set(
@@ -633,7 +665,6 @@ class SbListenerHelper:
                     if cat.catid not in include_ids:
                         logging.debug(f'Excluding sync filter "{cat.name}"')
                         self.ugs_filters.exclude_category(cat.catid)
-                        any_active_filters = True
                     else:
                         include_ids.remove(cat.catid)
 
@@ -650,21 +681,23 @@ class SbListenerHelper:
 
             logging.debug(f'Sync filter: {self.ugs_filters.map.as_array()}')
 
-        if self.sync_engine_cl:
-            engine_paths.append(f'{self.engine_dir.parent / "*"}')
-            engine_paths.append(f'{self.engine_dir / "..."}')
+        if sync_engine_cl:
+            assert engine_dir
+            engine_paths.append(f'{engine_dir.parent / "*"}')
+            engine_paths.append(f'{engine_dir / "..."}')
 
             pathspecs.extend(
-                [f'{path}@{self.sync_engine_cl}' for path in engine_paths])
+                [f'{path}@{sync_engine_cl}' for path in engine_paths])
 
-        if self.sync_project_cl:
-            project_paths.append(f'{self.uproj_path.parent / "..."}')
+        if sync_project_cl:
+            assert project_dir
+            project_paths.append(f'{project_dir / "..."}')
 
             for path in self.additional_paths_to_sync:
                 project_paths.append(f'//{self.p4client}{path}')
 
             pathspecs.extend(
-                [f'{path}@{self.sync_project_cl}' for path in project_paths])
+                [f'{path}@{sync_project_cl}' for path in project_paths])
 
         if len(pathspecs) < 1:
             self.parser.error('Nothing specified to sync')
@@ -673,22 +706,20 @@ class SbListenerHelper:
             for pathspec in pathspecs:
                 logging.info(f' - {pathspec}')
 
-        if not self.use_ugs:
-            logging.info('Determining CompatibleChangelist')
+        logging.info('Determining CompatibleChangelist')
 
-            # [:1] slice because we exclude AdditionalPathsToSync here
-            buildver_info = self.get_buildver_info(engine_paths, project_paths[:1])
+        # [:1] slice because we exclude AdditionalPathsToSync here
+        buildver_info = self.get_buildver_info(engine_paths, project_paths[:1])
 
-            if self.sync_engine_cl:
-                # Remove Build.version if present, so we can write our own.
-                if buildver_info.local_path and buildver_info.have_rev:
-                    pathspecs.append(f"{buildver_info.local_path}#0")
+        if sync_engine_cl:
+            # Remove Build.version if present, so we can write our own.
+            if buildver_info.local_path and buildver_info.have_rev:
+                pathspecs.append(f"{buildver_info.local_path}#0")
 
         remove_excluded_workload = SyncWorkload()
         if options.remove_excluded_workspace_files and self.ugs_filters:
             logging.info('Finding files in workspace that need to be removed.')
-            haves = p4_utils.p4_have(['//...'],
-                                     user=self.p4user, client=self.p4client)
+            haves = p4_utils.p4_have(['//...'], **p4kwargs)
             stream = self.p4clientspec['Stream']
             stream_path = pathlib.PurePosixPath(stream)
             num_exclusions = 0
@@ -710,7 +741,7 @@ class SbListenerHelper:
                 if workload.sync_files[buildver_path].rev != 0:
                     workload.remove(buildver_path)
 
-            if any_active_filters and self.ugs_filters:
+            if self.ugs_filters and self.ugs_filters.excluded_categories:
                 num_files = len(workload.sync_files)
                 logging.info(f'Applying sync filters... ({num_files:,} files)')
                 stream = self.p4clientspec['Stream']
@@ -730,25 +761,13 @@ class SbListenerHelper:
             for removal in remove_excluded_workload.sync_files.values():
                 workload.add_or_replace(removal)
 
-        def clobber_filter(clobber: SyncWorkload):
-            clobber.sync_files = dict(sorted(clobber.sync_files.items()))
+        sync_workload = p4_get_preview_sync_files(pathspecs, **p4kwargs)
+        sync_filter(sync_workload)
 
-            for depot_path, clobber_file in clobber.sync_files.items():
-                if clobber_file.client_path is not None:
-                    if self.clobber_engine and self.engine_dir:
-                        if path_is_relative_to(self.engine_dir,
-                                               clobber_file.client_path):
-                            clobber_file.clobber = True
+        return sync_workload, buildver_info
 
-                    if self.clobber_project and self.uproj_path:
-                        if path_is_relative_to(self.uproj_path.parent,
-                                               clobber_file.client_path):
-                            clobber_file.clobber = True
-                else:
-                    # Shouldn't happen.
-                    logging.error('clobber_filter(): Missing client path for '
-                                  f'clobber file: {depot_path}')
-                    continue
+    def run_sync(self, options: argparse.Namespace) -> int:
+        self.check_sync_options(options)
 
         sync_result = None
         if self.use_ugs:
@@ -761,32 +780,62 @@ class SbListenerHelper:
                 client=self.p4client
             )
         else:
+            project_dir = self.uproj_path.parent if self.uproj_path else None
+
+            sync_workload, buildver_info = self._get_sync_workload(options,
+                engine_dir=self.engine_dir, sync_engine_cl=self.sync_engine_cl,
+                project_dir=project_dir, sync_project_cl=self.sync_project_cl,
+                **self.p4kwargs)
+
+            # User may have requested to write the pending sync files and revisions to
+            # file and exit early.
+            write_workload_file: io.TextIOWrapper | None = options.dump_sync
+            if write_workload_file is not None:
+                for item in sync_workload.sync_files.values():
+                    write_workload_file.write(f'{item.spec}\n')
+                return 0
+
+            def clobber_filter(clobber: SyncWorkload):
+                clobber.sync_files = dict(sorted(clobber.sync_files.items()))
+
+                for depot_path, clobber_file in clobber.sync_files.items():
+                    if clobber_file.client_path is not None:
+                        if self.clobber_engine and self.engine_dir:
+                            if path_is_relative_to(self.engine_dir, clobber_file.client_path):
+                                clobber_file.clobber = True
+
+                        if self.clobber_project and self.uproj_path:
+                            if path_is_relative_to(self.uproj_path.parent, clobber_file.client_path):
+                                clobber_file.clobber = True
+                    else:
+                        # Shouldn't happen.
+                        logging.error(f'clobber_filter(): Missing client path for clobber file: {depot_path}')
+                        continue
+
+            if not self.dry_run:
+                self.write_sync_state()
+
             sync_result = p4_sync(
-                pathspecs,
-                sync_filter_fn=sync_filter,
+                sync_workload,
                 clobber_filter_fn=clobber_filter,
                 progress_fn=self.on_sync_progress,
                 dry_run=self.dry_run,
+                port=self.p4port,
                 user=self.p4user,
                 client=self.p4client,
-                write_workload_file=options.dump_sync,
             )
 
-        if sync_result is None:
-            return 0  # Nothing was synced
-        elif self.dry_run or options.dump_sync:
-            return sync_result  # Nothing more to do
+            if sync_result is None or self.dry_run:
+                return 0 # Nothing needed to be synced
 
-        # Write updated Build.version.
-        if not self.use_ugs:
+            # Write updated Build.version.
             if buildver_info.local_path and buildver_info.updated_text:
                 logging.info('Updating Build.version')
                 try:
-                    with open(buildver_info.local_path, 'wt') as buildver_file:
+                    with open(buildver_info.local_path, 'wt', encoding='utf-8') as buildver_file:
                         buildver_file.write(buildver_info.updated_text)
                 except Exception as exc:
-                    logging.error('Exception writing to Build.version',
-                                  exc_info=exc)
+                    logging.error('Exception writing to Build.version', exc_info=exc)
 
         # Generate project files.
         if self.generate_after_sync:
@@ -795,12 +844,194 @@ class SbListenerHelper:
             if gpf_result == 0:
                 logging.info('Done generating project files')
             else:
-                logging.error('GenerateProjectFiles failed with code '
-                              f'{gpf_result}')
+                logging.error(f'GenerateProjectFiles failed with code {gpf_result}')
                 return gpf_result
 
         return sync_result
 
+    STATE_FILE = '.sblsync.json'
+
+    @property
+    def engine_sync_state_path(self) -> pathlib.Path | None:
+        if self.engine_dir:
+            return self.engine_dir.parent / self.STATE_FILE
+        else:
+            return None
+
+    @property
+    def project_sync_state_path(self) -> pathlib.Path | None:
+        if self.uproj_path:
+            return self.uproj_path.parent / self.STATE_FILE
+        else:
+            return None
+
+    @dataclass
+    class SyncState:
+        changelist: int
+        sync_time: datetime.datetime
+
+    def _save_sync_state_to_file(self, sync_state: SyncState, path: pathlib.Path):
+        data = {
+            'changelist': sync_state.changelist,
+            'sync_time': sync_state.sync_time.isoformat(),
+        }
+
+        try:
+            with open(path, 'wt', encoding='utf-8') as f:
+                json.dump(data, f, indent=4)
+        except Exception as exc:
+            logging.error(f'Unable to write sync state to {path}', exc_info=exc)
+            return
+
+        logging.debug(f'Sync state written to {path}: {data}')
+
+    def _load_sync_state_from_file(self, path: pathlib.Path) -> SyncState | None:
+        if not path.exists():
+            logging.debug(f'No sync state located at {path}')
+            return None
+
+        try:
+            with open(path, 'r', encoding='utf-8') as f:
+                state_json: dict = json.load(f)
+                logging.debug(f'Sync state loaded from {path}: {state_json}')
+        except Exception as exc:
+            logging.error(f'Unable to read sync state from {path}', exc_info=exc)
+            return None
+
+        return SbListenerHelper.SyncState(
+            changelist=int(state_json['changelist']),
+            sync_time=datetime.datetime.fromisoformat(state_json['sync_time'])
+        )
+
+    def read_engine_sync_state(self) -> SyncState | None:
+        if self.engine_sync_state_path:
+            return self._load_sync_state_from_file(self.engine_sync_state_path)
+
+    def read_project_sync_state(self) -> SyncState | None:
+        if self.project_sync_state_path:
+            return self._load_sync_state_from_file(self.project_sync_state_path)
+
+    def write_sync_state(self):
+        sync_time = datetime.datetime.now()
+
+        if self.sync_engine_cl and self.engine_sync_state_path:
+            engine_state = SbListenerHelper.SyncState(self.sync_engine_cl, sync_time)
+            self._save_sync_state_to_file(engine_state, self.engine_sync_state_path)
+
+        if self.sync_project_cl and self.project_sync_state_path:
+            project_state = SbListenerHelper.SyncState(self.sync_project_cl, sync_time)
+            self._save_sync_state_to_file(project_state, self.project_sync_state_path)
+
+    def run_syncstatus(self, options: argparse.Namespace) -> int:
+        self.check_p4_options(options)
+
+        if not self.engine_dir and not self.uproj_path:
+            self.parser.error('Neither engine nor project path were provided')
+
+        project_dir = self.uproj_path.parent if self.uproj_path else None
+
+        # Read the notionally synced changelists, as written during the last attempted sync.
+        engine_state = self.read_engine_sync_state()
+        project_state = self.read_project_sync_state()
+
+        def guess_have_cl(path: pathlib.Path, client: str | None = None) -> int | None:
+            logging.debug(f'{self.STATE_FILE} not available for {path}; falling back to `p4 cstat #have`')
+
+            if not path.is_dir():
+                return None
+
+            p4kwargs = {**self.p4kwargs}
+            if client:
+                p4kwargs['client'] = client
+
+            try:
+                orig_cwd = os.getcwd()  # Ensure user's P4CONFIG files are read appropriately
+                os.chdir(path)
+                have_changes = p4_utils.p4_get_decoded_records('cstat', [f'{path / "...#have"}'],
+                                                               **p4kwargs)
+            finally:
+                os.chdir(orig_cwd)
+
+            have_changes = sorted(have_changes, key=lambda x: int(x.get('change', -1)), reverse=True)
+
+            if have_changes:
+                if change := have_changes[0].get('change'):
+                    return int(change)
+
+            return None
+
+        result = {}
+
+        if self.engine_dir:
+            if eng_detected_ws := p4_utils.p4_get_var('P4CLIENT', self.engine_dir):
+                result['engine_detected_ws'] = eng_detected_ws
+
+            if engine_state:
+                self.sync_engine_cl = engine_state.changelist
+                result['engine_cl'] = engine_state.changelist
+                result['engine_sync_time'] = engine_state.sync_time.isoformat()
+            else:
+                guess_eng_cl = guess_have_cl(self.engine_dir)
+                self.sync_engine_cl = guess_eng_cl
+                result['engine_cl'] = guess_eng_cl
+                result['engine_sync_time'] = '(Unknown)'
+
+        if project_dir:
+            if proj_detected_ws := p4_utils.p4_get_var('P4CLIENT', project_dir):
+                result['project_detected_ws'] = proj_detected_ws
+
+            if project_state:
+                self.sync_project_cl = project_state.changelist
+                result['project_cl'] = project_state.changelist
+                result['project_sync_time'] = project_state.sync_time.isoformat()
+            else:
+                guess_proj_cl = guess_have_cl(project_dir, self.p4client_project)
+                self.sync_project_cl = guess_proj_cl
+                result['project_cl'] = guess_proj_cl
+                result['project_sync_time'] = '(Unknown)'
+
+        # Preview whether attempting to sync to the notional changelists would actually result in any changes.
+        # If the workspace already reflects the notional changelists, the "workload" should be empty
+        # (i.e. "//...@12345678 - file(s) up-to-date.").
+        #
+        # If syncing would have taken an action, the workspace is considered dirty.
+
+        try:
+            orig_cwd = os.getcwd()  # Ensure user's P4CONFIG files are read appropriately
+            if self.p4client_project and self.p4client_project != self.p4client:
+                # Potentially two queries if we specify a separate project workspace.
+                either_dirty = False
+
+                if self.engine_dir:
+                    os.chdir(self.engine_dir)
+                    eng_workload, _ = self._get_sync_workload(options,
+                        engine_dir=self.engine_dir, sync_engine_cl=self.sync_engine_cl,
+                        port=self.p4port, user=self.p4user, client=self.p4client)
+                    either_dirty = either_dirty or len(eng_workload.sync_files) > 0
+
+                if project_dir:
+                    os.chdir(project_dir)
+                    proj_workload, _ = self._get_sync_workload(options,
+                        project_dir=project_dir, sync_project_cl=self.sync_project_cl,
+                        port=self.p4port, user=self.p4user, client=self.p4client_project)
+                    either_dirty = either_dirty or len(proj_workload.sync_files) > 0
+
+                result['dirty'] = either_dirty
+            else:
+                sync_workload, _ = self._get_sync_workload(options,
+                    engine_dir=self.engine_dir, sync_engine_cl=self.sync_engine_cl,
+                    project_dir=project_dir, sync_project_cl=self.sync_project_cl,
+                    **self.p4kwargs)
+
+                result['dirty'] = len(sync_workload.sync_files) > 0
+        finally:
+            os.chdir(orig_cwd)
+
+        result['log'] = self.log_string.getvalue()
+        json.dump(result, sys.stdout, indent=4)
+
+        return 0
+
     @dataclass
     class BuildVerInfo:
         local_path: Optional[str] = None
@@ -831,9 +1062,7 @@ class SbListenerHelper:
 
         ret_info.local_path = buildver_local_path
 
-        p4opts = {'user': self.p4user, 'client': self.p4client}
-
-        record = p4_utils.p4_have([buildver_local_path], **p4opts)[0]
+        record = p4_utils.p4_have([buildver_local_path], **self.p4kwargs)[0]
         ret_info.have_rev = int(record['haveRev']) if (
             record['code'] == 'stat') else None
 
@@ -843,7 +1072,7 @@ class SbListenerHelper:
         prints = p4_utils.p4_print(
             [f'{path}@{self.sync_engine_cl}'
              for path in (epicint_local_path, buildver_local_path)],
-            **p4opts)
+            **self.p4kwargs)
         epicint_print_result = prints[0]
         buildver_print_result = prints[1]
 
@@ -874,7 +1103,7 @@ class SbListenerHelper:
             if quick_lower > 0:
                 quick_range = f'{quick_lower},{upper_cl}'
                 check_lower = p4_utils.p4_latest_code_change(
-                    paths, in_range=quick_range, exts=['.cpp'], **p4opts)
+                    paths, in_range=quick_range, exts=['.cpp'], **self.p4kwargs)
 
                 if check_lower:
                     if in_lower_cl:
@@ -905,14 +1134,14 @@ class SbListenerHelper:
             eng_code_cl = eng_upper_bound_cl
         else:
             eng_code_cl = p4_utils.p4_latest_code_change(
-                engine_paths, in_range=eng_range, **p4opts)
+                engine_paths, in_range=eng_range, **self.p4kwargs)
 
         if proj_upper_bound_cl:
             if lower_bound_cl and (lower_bound_cl >= proj_upper_bound_cl):
                 proj_code_cl = proj_upper_bound_cl
             else:
                 proj_code_cl = p4_utils.p4_latest_code_change(
-                    project_paths, in_range=proj_range, **p4opts)
+                    project_paths, in_range=proj_range, **self.p4kwargs)
         else:
             proj_code_cl = None
 
@@ -955,7 +1184,7 @@ class SbListenerHelper:
         # Read existing local file and skip update if unchanged
         existing_build_ver = None
         try:
-            with open(buildver_local_path, 'r') as f:
+            with open(buildver_local_path, 'r', encoding='utf-8') as f:
                 existing_build_ver = json.load(f)
         except Exception as exc:
             logging.warning('Exception reading local Build.version',
@@ -1000,28 +1229,60 @@ class SbListenerHelper:
             ):
                 logging.info(progress_str)
 
+    def check_p4_options(self, options: argparse.Namespace):
+        self.p4port = options.p4port
+        self.p4user = options.p4user
+        self.p4client = options.p4client
+        self.p4client_project = options.p4client_project
+
+        if self.p4port:
+            self.p4kwargs['port'] = self.p4port
+
+        if self.p4user:
+            self.p4kwargs['user'] = self.p4user
+
+        if self.p4client:
+            self.p4kwargs['client'] = self.p4client
+
+        if p4login := p4_utils.p4_get_login(**self.p4kwargs):
+            logging.info(f'Perforce: Logged in as {p4login}')
+        else:
+            self.parser.error('No valid Perforce session found. Make sure you are logged into Perforce.')
+
+        if options.project:
+            try:
+                self.uproj_path = options.project.resolve(strict=True)
+            except FileNotFoundError:
+                self.parser.error('argument --project: FileNotFoundError: '
+                                  f'{options.project}')
+
+            try:
+                _ = json.load(self.uproj_path.open())
+            except json.JSONDecodeError:
+                self.parser.error('argument --project: unable to parse JSON'
+                                  f'\n\n{traceback.format_exc()}')
+
+        self.engine_dir = self.find_engine_dir(options)
+
     def check_sync_options(self, options: argparse.Namespace):
         '''
         Parse `sync` options, caching their values, and ensuring consistency
         between related options.
         '''
+        self.check_p4_options(options)
+
         self.dry_run = options.dry_run
         self.clobber_engine = options.clobber_engine
         self.clobber_project = options.clobber_project
         self.ugs_versioning = options.ugs_versioning
         self.ugs_lib_dir = options.ugs_lib_dir
 
-        self.p4user = options.p4user
-        self.p4client = options.p4client
-        self.p4clientspec = p4_utils.p4_get_client(self.p4client,
-                                                   user=self.p4user,
-                                                   client=self.p4client) or {}
+        if self.p4client:
+            self.p4clientspec = p4_utils.p4_get_client(self.p4client, **self.p4kwargs) or {}
 
         if options.project is None:
             if options.project_cl is not None:
                 self.parser.error('--project-cl requires --project')
-            elif options.generate:
-                self.parser.error('--generate requires --project')
 
         if options.engine_cl is not None and options.engine_cl <= 0:
             self.parser.error(f'Invalid --engine-cl {options.engine_cl}')
@@ -1043,25 +1304,10 @@ class SbListenerHelper:
                               '`--engine-cl`; these arguments are mutally '
                               'exclusive and should not be used together.')
 
-        if options.project:
-            try:
-                self.uproj_path = options.project.resolve(strict=True)
-            except FileNotFoundError:
-                self.parser.error('argument --project: FileNotFoundError: '
-                                  f'{options.project}')
-
-            try:
-                _ = json.load(self.uproj_path.open())
-            except json.JSONDecodeError:
-                self.parser.error('argument --project: unable to parse JSON'
-                                  f'\n\n{traceback.format_exc()}')
-
         engine_dir_required = ((self.sync_engine_cl is not None)
                                or (self.generate_after_sync))
 
         if engine_dir_required:
-            self.engine_dir = self.find_engine_dir(options)
-
             if self.engine_dir is None:
                 if self.sync_engine_cl is not None:
                     self.parser.error('--engine-cl was specified, but unable '
@@ -1083,11 +1329,12 @@ class SbListenerHelper:
             while not iter_dir.samefile(iter_dir.anchor):
                 candidate = iter_dir / 'Engine'
                 if candidate.is_dir():
+                    logging.debug(f'Engine directory: {candidate}')
                     return candidate
                 iter_dir = iter_dir.parent
 
         # TODO: Consider looking for Engine dir as ancestor of this script?
-
+        logging.debug('Engine directory not found')
         return None
 
     def generate_project_files(self) -> int:
@@ -1095,7 +1342,7 @@ class SbListenerHelper:
         Invokes GenerateProjectFiles script, streams the output to INFO
         logging, and returns the exit code.
         '''
-        assert self.uproj_path and self.engine_dir
+        assert self.engine_dir
 
         if sys.platform.startswith('win'):
             gpf_script = os.path.join(self.engine_dir, 'Build', 'BatchFiles',
@@ -1104,10 +1351,11 @@ class SbListenerHelper:
             gpf_script = os.path.join(self.engine_dir, 'Build', 'BatchFiles',
                                       'Linux', 'GenerateProjectFiles.sh')
 
-        args = [gpf_script, str(self.uproj_path)]
+        args = [gpf_script]
+        if self.uproj_path:
+            args.append(str(self.uproj_path))
 
-        logging.debug('generate_project_files(): invoking subprocess: '
-                      f'args={args}')
+        logging.debug(f'generate_project_files(): invoking subprocess: args={args}')
 
         try:
             with subprocess.Popen(args, stdout=subprocess.PIPE) as proc:
diff --git a/Engine/Plugins/VirtualProduction/Switchboard/Source/Switchboard/switchboard/__main__.py b/Engine/Plugins/VirtualProduction/Switchboard/Source/Switchboard/switchboard/__main__.py
--- a/Engine/Plugins/VirtualProduction/Switchboard/Source/Switchboard/switchboard/__main__.py
+++ b/Engine/Plugins/VirtualProduction/Switchboard/Source/Switchboard/switchboard/__main__.py
@@ -40,7 +40,7 @@ def launch():
 
     # Global variables can be inited after QT has been started (they may show dialogues if there are errors)
     SETTINGS.init()
-    CONFIG.init(SETTINGS.CONFIG)
+    CONFIG.init_with_file_path(SETTINGS.CONFIG)
 
     # parse arguments
     args = parse_arguments()
diff --git a/Engine/Plugins/VirtualProduction/Switchboard/Source/Switchboard/switchboard/config.py b/Engine/Plugins/VirtualProduction/Switchboard/Source/Switchboard/switchboard/config.py
--- a/Engine/Plugins/VirtualProduction/Switchboard/Source/Switchboard/switchboard/config.py
+++ b/Engine/Plugins/VirtualProduction/Switchboard/Source/Switchboard/switchboard/config.py
@@ -2423,8 +2423,8 @@ class Config(object):
         '''
         self.saving_allowed = self.saving_allowed_fifo.pop()
 
-    def init(self, file_path: Union[str, pathlib.Path]):
-        self.init_with_file_path(file_path)
+    def __init__(self):
+        self.file_path: pathlib.Path | None = None
 
     def init_with_file_path(self, file_path: Union[str, pathlib.Path]):
         if file_path:
@@ -2463,7 +2463,7 @@ class Config(object):
             + [setting for _, setting in self.source_control_settings.items()] \
             + [setting for _, setting in self.unreal_insight_settings.items()] \
             + [setting for _, setting in self.mu_settings.items()]
-            # TODO: multiuser
+
         for setting in all_settings:
             setting.signal_setting_changed.connect(lambda: self.save())
             setting.signal_setting_overridden.connect(
diff --git a/Engine/Plugins/VirtualProduction/Switchboard/Source/Switchboard/switchboard/devices/device_base.py b/Engine/Plugins/VirtualProduction/Switchboard/Source/Switchboard/switchboard/devices/device_base.py
--- a/Engine/Plugins/VirtualProduction/Switchboard/Source/Switchboard/switchboard/devices/device_base.py
+++ b/Engine/Plugins/VirtualProduction/Switchboard/Source/Switchboard/switchboard/devices/device_base.py
@@ -103,8 +103,8 @@ class Device(QtCore.QObject):
                 override = kwargs[setting.attr_name]
                 setting.override_value_from_config(self.name, override)
 
-        self._project_changelist: Optional[int] = None
-        self._engine_changelist: Optional[int] = None
+        self._project_changelist: Optional[int | str] = None
+        self._engine_changelist: Optional[int | str] = None
 
         self._status = DeviceStatus.DISCONNECTED
 
@@ -265,7 +265,7 @@ class Device(QtCore.QObject):
         return self._project_changelist
 
     @project_changelist.setter
-    def project_changelist(self, value: Optional[int]):
+    def project_changelist(self, value: Optional[int | str]):
         self._project_changelist = value
         self.device_qt_handler.signal_device_project_changelist_changed.emit(
             self)
@@ -275,7 +275,7 @@ class Device(QtCore.QObject):
         return self._engine_changelist
 
     @engine_changelist.setter
-    def engine_changelist(self, value: Optional[int]):
+    def engine_changelist(self, value: Optional[int | str]):
         self._engine_changelist = value
         self.device_qt_handler.signal_device_engine_changelist_changed.emit(
             self)
diff --git a/Engine/Plugins/VirtualProduction/Switchboard/Source/Switchboard/switchboard/devices/livelinkhub/plugin_livelinkhub.py b/Engine/Plugins/VirtualProduction/Switchboard/Source/Switchboard/switchboard/devices/livelinkhub/plugin_livelinkhub.py
--- a/Engine/Plugins/VirtualProduction/Switchboard/Source/Switchboard/switchboard/devices/livelinkhub/plugin_livelinkhub.py
+++ b/Engine/Plugins/VirtualProduction/Switchboard/Source/Switchboard/switchboard/devices/livelinkhub/plugin_livelinkhub.py
@@ -224,12 +224,14 @@ class DeviceLiveLinkHub(DeviceUnreal):
         *,
         program_name: str,
         returncode: int,
-        get_stdout_str: Callable[[], str]
+        get_stdout_str: Callable[[], str],
+        get_stderr_str: Callable[[], str],
     ):
         super().do_program_ended_update(
             program_name=program_name,
             returncode=returncode,
             get_stdout_str=get_stdout_str,
+            get_stderr_str=get_stderr_str
         )
 
         if program_name == self.LLH_PROG_NAME:
diff --git a/Engine/Plugins/VirtualProduction/Switchboard/Source/Switchboard/switchboard/devices/ndisplayTest/plugin_ndisplayTest.py b/Engine/Plugins/VirtualProduction/Switchboard/Source/Switchboard/switchboard/devices/ndisplayTest/plugin_ndisplayTest.py
--- a/Engine/Plugins/VirtualProduction/Switchboard/Source/Switchboard/switchboard/devices/ndisplayTest/plugin_ndisplayTest.py
+++ b/Engine/Plugins/VirtualProduction/Switchboard/Source/Switchboard/switchboard/devices/ndisplayTest/plugin_ndisplayTest.py
@@ -731,7 +731,8 @@ class DevicenDisplayTest(DevicenDisplay):
     def do_program_ended_update(self, *,
         program_name: str,
         returncode: int,
-        get_stdout_str: Callable[[], str]
+        get_stdout_str: Callable[[], str],
+        get_stderr_str: Callable[[], str],
     ):
         ''' [override] Handles test & transfer finish '''
 
@@ -743,12 +744,12 @@ class DevicenDisplayTest(DevicenDisplay):
         super().do_program_ended_update(
             program_name=program_name,
             returncode=returncode,
-            get_stdout_str=get_stdout_str)
+            get_stdout_str=get_stdout_str,
+            get_stderr_str=get_stderr_str)
 
         # Handle test finish
         if program_name == UnrealJobs.Unreal.value:
             self.start_retrieve_report(unreal_exit_code=returncode)
-            pass
 
     def start_retrieve_log(self, unreal_exit_code: int):
         ''' [override] Retrieve the log file if logging is enabled. '''
diff --git a/Engine/Plugins/VirtualProduction/Switchboard/Source/Switchboard/switchboard/devices/unreal/plugin_unreal.py b/Engine/Plugins/VirtualProduction/Switchboard/Source/Switchboard/switchboard/devices/unreal/plugin_unreal.py
--- a/Engine/Plugins/VirtualProduction/Switchboard/Source/Switchboard/switchboard/devices/unreal/plugin_unreal.py
+++ b/Engine/Plugins/VirtualProduction/Switchboard/Source/Switchboard/switchboard/devices/unreal/plugin_unreal.py
@@ -55,9 +55,7 @@ class UnrealJobs(Enum):
     BuildPrefix = 'build_'
     BuildProject = 'build_project'
 
-    CstatPrefix = 'cstat'
-    CstatProject = 'cstat_project'
-    CstatEngine = 'cstat_engine'
+    SyncStatus = 'syncstatus'
 
     Sync = 'sync'
     BuildDDC = 'unreal_ddc'
@@ -842,33 +840,33 @@ class DeviceUnreal(Device):
 
     @classmethod
     def p4_refresh_ugs_config(cls):
-        if cls.ugs_config is None:
+        if DeviceUnreal.ugs_config is None:
             # One-time initialization
             def refresh():
-                cls.p4_refresh_ugs_config()
+                DeviceUnreal.p4_refresh_ugs_config()
 
             CONFIG.P4_ENABLED.signal_setting_changed.connect(refresh)
             CONFIG.P4_ENGINE_PATH.signal_setting_changed.connect(refresh)
             CONFIG.P4_PROJECT_PATH.signal_setting_changed.connect(refresh)
 
         if not CONFIG.P4_ENABLED.get_value():
-            cls.ugs_config = ugs_utils.IniParser()
+            DeviceUnreal.ugs_config = ugs_utils.IniParser()
             return
 
-        if cls._pending_ugs_config_refresh:
+        if DeviceUnreal._pending_ugs_config_refresh:
             return
-        elif cls._ugs_config_runnable:
-            cls._pending_ugs_config_refresh = True
+        elif DeviceUnreal._ugs_config_runnable:
+            DeviceUnreal._pending_ugs_config_refresh = True
             return
 
         engine_dir = CONFIG.P4_ENGINE_PATH.get_value() or None
         project_dir = CONFIG.P4_PROJECT_PATH.get_value() or None
-        cls._ugs_config_runnable = cls.UgsConfigRunnable(
+        DeviceUnreal._ugs_config_runnable = DeviceUnreal.UgsConfigRunnable(
             engine_dir, project_dir)
 
-        cls._ugs_config_runnable.signals.result.connect(
-            cls._on_updated_ugs_config)
-        QtCore.QThreadPool.globalInstance().start(cls._ugs_config_runnable)
+        DeviceUnreal._ugs_config_runnable.signals.result.connect(
+            DeviceUnreal._on_updated_ugs_config)
+        QtCore.QThreadPool.globalInstance().start(DeviceUnreal._ugs_config_runnable)
 
     @classmethod
     def _on_updated_ugs_config(cls, ugs_config: ugs_utils.IniParser):
@@ -1070,11 +1068,11 @@ class DeviceUnreal(Device):
 
         self.runtime_str = ""
 
-        # The CLs returned via p4 cstat query, or optimistically
+        # The CLs returned via syncstatus query, or optimistically
         # updated with the "inflight" values upon successful sync
         # (Defined in the base class, documented here for context)
-        #self._project_changelist: Optional[int] = None
-        #self._engine_changelist: Optional[int] = None
+        #self._project_changelist: Optional[int | str] = None
+        #self._engine_changelist: Optional[int | str] = None
 
         # The CLs for which outstanding p4 syncs have been queued
         self.inflight_project_cl = None
@@ -1170,7 +1168,7 @@ class DeviceUnreal(Device):
                 'rsync_port'].signal_setting_changed.connect(
                     launch_rsync_server)
 
-        if self.ugs_config is None:
+        if self.ugs_config is None and self._ugs_config_runnable is None:
             self.p4_refresh_ugs_config()
 
         app = QtCore.QCoreApplication.instance()
@@ -1476,7 +1474,7 @@ class DeviceUnreal(Device):
         if engine_sync_method == EngineSyncMethod.Build_Engine.value:
             self.widget.engine_changelist_label.show()
             if not self.is_disconnected:
-                self._request_engine_changelist_number()
+                self._request_syncstatus()
         else:
             self.widget.engine_changelist_label.hide()
             self.widget.update_build_info(self)
@@ -1550,93 +1548,44 @@ class DeviceUnreal(Device):
         _, msg = message_protocol.create_copy_file_from_listener_message(proj_target_path)
         self.unreal_client.send_message(msg)
 
-
-    def _request_project_changelist_number(self):
+    def _request_syncstatus(self):
         if not CONFIG.P4_ENABLED.get_value():
             return
 
-        client_name = CONFIG.SOURCE_CONTROL_WORKSPACE.get_value(self.name)
+        project_path = CONFIG.UPROJECT_PATH.get_value(self.name)
+        engine_dir = CONFIG.ENGINE_DIR.get_value(self.name)
+        workspace = CONFIG.SOURCE_CONTROL_WORKSPACE.get_value(self.name)
         project_workspace = CONFIG.PROJECT_WORKSPACE.get_value(self.name)
 
-        if project_workspace:
-            client_name = project_workspace
-
-        if not client_name:
-            LOGGER.warning(
-                f"{self.name}: Missing workspace name to query the project "
-                "changelist")
-            return
-
-        p4_path = CONFIG.P4_PROJECT_PATH.get_value()
-
-        if not p4_path:
-            LOGGER.warning(
-                f"{self.name}: Missing p4 path to query the project "
-                "changelist")
-            return
-
-        formatstring = "%change%"
-        args = f'-F "{formatstring}" -c {client_name} cstat {p4_path}/...#have'
-
-        program_name = UnrealJobs.CstatProject.value
-
-        working_dir = os.path.dirname(
-            CONFIG.UPROJECT_PATH.get_value(self.name))
-
-        puuid, msg = message_protocol.create_start_process_message(
-            prog_path="p4",
-            prog_args=args,
-            prog_name=program_name,
-            caller=self.name,
-            working_dir=working_dir,
-            update_clients_with_stdout=False,
-            hide=True,
-        )
-
-        self.program_start_queue.add(
-            ProgramStartQueueItem(
-                name=program_name,
-                puuid_dependency=None,
-                puuid=puuid,
-                msg_to_unreal_client=msg,
-            ),
-            unreal_client=self.unreal_client,
-        )
-
-    def _request_engine_changelist_number(self):
-        if not CONFIG.P4_ENABLED.get_value():
-            return
-
-        client_name = CONFIG.SOURCE_CONTROL_WORKSPACE.get_value(self.name)
-
-        if not client_name:
-            LOGGER.warning(
-                f"{self.name}: Missing workspace name to query the engine "
-                "changelist")
-            return
+        syncstatus_args = [
+            f'"{self.sbl_helper_path}"',
+            # '--log-level=DEBUG',
+            'syncstatus',
+            f'--engine-dir="{engine_dir}"',
+            f'--project="{project_path}"',
+        ]
 
-        p4_path = CONFIG.P4_ENGINE_PATH.get_value()
+        if CONFIG.ENGINE_SYNC_METHOD.get_value() == EngineSyncMethod.Build_Engine.value:
+            syncstatus_args.append(f'--engine-dir="{engine_dir}"')
 
-        if not p4_path:
-            LOGGER.warning(
-                f'{self.name}: Missing p4 path to query the engine changelist')
-            return
+        if workspace:
+            syncstatus_args.append(f'--p4client={workspace}')
 
-        formatstring = "%change%"
-        args = f'-F "{formatstring}" -c {client_name} cstat {p4_path}/...#have'
+        if project_workspace:
+            syncstatus_args.append(f'--p4client-project={project_workspace}')
 
-        program_name = UnrealJobs.CstatEngine.value
+        program_name = UnrealJobs.SyncStatus.value
 
-        working_dir = os.path.dirname(
-            CONFIG.UPROJECT_PATH.get_value(self.name))
+        working_dir = os.path.dirname(CONFIG.UPROJECT_PATH.get_value(self.name))
 
         puuid, msg = message_protocol.create_start_process_message(
-            prog_path="p4",
-            prog_args=args,
+            prog_path=self.python_path,
+            prog_args=' '.join(syncstatus_args),
             prog_name=program_name,
             caller=self.name,
             working_dir=working_dir,
             update_clients_with_stdout=False,
+            combined_stdout_stderr=False,
             hide=True,
         )
 
@@ -1681,13 +1630,31 @@ class DeviceUnreal(Device):
         super().disconnect_listener()
         self.unreal_client.disconnect()
 
-    def sync(
-            self,
-            engine_cl: Optional[int],
-            project_cl: Optional[int],
-            puuid_dependency: Optional[uuid.UUID] = None
-            ) -> Optional[uuid.UUID]:
+    @property
+    def python_path(self) -> str:
+        engine_dir = CONFIG.ENGINE_DIR.get_value(self.name)
+        if self.platform_binary_directory == 'Win64':
+            return os.path.normpath(os.path.join(
+                engine_dir, 'Binaries', 'ThirdParty', 'Python3',
+                self.platform_binary_directory, 'python.exe'))
+        else:
+            return os.path.normpath(os.path.join(
+                engine_dir, 'Binaries', 'ThirdParty', 'Python3',
+                self.platform_binary_directory, 'bin', 'python3'))
+
+    @property
+    def sbl_helper_path(self) -> str:
+        engine_dir = CONFIG.ENGINE_DIR.get_value(self.name)
+        return os.path.normpath(os.path.join(
+            engine_dir, 'Plugins', 'VirtualProduction', 'Switchboard',
+            'Source', 'Switchboard', 'sbl_helper.py'))
 
+    def sync(
+        self,
+        engine_cl: Optional[int | str],
+        project_cl: Optional[int | str],
+        puuid_dependency: Optional[uuid.UUID] = None
+    ) -> Optional[uuid.UUID]:
         if not engine_cl and not project_cl:
             LOGGER.warning(
                 "Neither project nor engine changelist is selected. "
@@ -1735,23 +1702,8 @@ class DeviceUnreal(Device):
             f"{self.name}: Queuing sync for project {project_name} "
             f"(revisions: engine={engine_cl}, project={project_cl})")
 
-        if self.platform_binary_directory == 'Win64':
-            python_path = os.path.normpath(os.path.join(
-                engine_dir, 'Binaries', 'ThirdParty', 'Python3',
-                self.platform_binary_directory, 'python.exe'))
-        else:
-            python_path = os.path.normpath(os.path.join(
-                engine_dir, 'Binaries', 'ThirdParty', 'Python3',
-                self.platform_binary_directory, 'bin', 'python3'))
-
-        helper_path = os.path.normpath(os.path.join(
-            engine_dir, 'Plugins', 'VirtualProduction', 'Switchboard',
-            'Source', 'Switchboard', 'sbl_helper.py'))
-
-        sync_tool = python_path
-
         sync_args = [
-            f'"{helper_path}"',
+            f'"{self.sbl_helper_path}"',
             # '--log-level=DEBUG',
             'sync',
             f'--project="{project_path}"',
@@ -1829,7 +1781,7 @@ class DeviceUnreal(Device):
                 self.unreal_client.send_message(msg)
 
         puuid, msg = message_protocol.create_start_process_message(
-            prog_path=sync_tool,
+            prog_path=self.python_path,
             prog_args=sync_args,
             prog_name=program_name,
             caller=self.name,
@@ -1847,7 +1799,7 @@ class DeviceUnreal(Device):
                 msg_to_unreal_client=msg,
                 launch_fn=lambda: LOGGER.info(
                     f"{self.name}: Sending sync command: "
-                    f"{sync_tool} {sync_args}")
+                    f"{self.python_path} {sync_args}")
             ),
             unreal_client=self.unreal_client,
         )
@@ -2576,6 +2528,10 @@ class DeviceUnreal(Device):
             b64bytes = base64.b64decode(message['stdoutB64'])
             return b64bytes.decode()
 
+        def get_stderr_str() -> str:
+            b64bytes = base64.b64decode(message['stderrB64'])
+            return b64bytes.decode()
+
         LOGGER.info(
             f"{self.name}: Program with id {puuid} exited with "
             f"returncode {returncode}")
@@ -2592,6 +2548,7 @@ class DeviceUnreal(Device):
             program_name=program_name,
             returncode=returncode,
             get_stdout_str=get_stdout_str,
+            get_stderr_str=get_stderr_str,
         )
 
     def do_program_ended_update(
@@ -2599,7 +2556,8 @@ class DeviceUnreal(Device):
         *,
         program_name: str,
         returncode: int,
-        get_stdout_str: Callable[[], str]
+        get_stdout_str: Callable[[], str],
+        get_stderr_str: Callable[[], str],
     ):
         # Update the widget classes (e.g. open button icon) once the
         # program_start_queue has been updated, since the status of the
@@ -2697,48 +2655,44 @@ class DeviceUnreal(Device):
 
                 # Forces an update to the changelist field (to hide the
                 # Building state).
-                self._request_project_changelist_number()
+                self._request_syncstatus()
                 if CONFIG.ENGINE_SYNC_METHOD.get_value() == EngineSyncMethod.Build_Engine.value:
-                    self._request_engine_changelist_number()
                     self._request_version_files()
 
-        elif UnrealJobs.CstatPrefix.value in program_name:
+        elif program_name == UnrealJobs.SyncStatus.value:
             output = get_stdout_str()
-            changelists = [line.strip() for line in output.split()]
 
             try:
-                current_changelist = str(int(changelists[-1]))
+                syncstatus: dict = json.loads(output)
+            except Exception:
+                stderr = ''
+                try:
+                    stderr = get_stderr_str()
+                except KeyError:
+                    pass
 
-                # when not connected to p4, you get a message similar to:
-                #  "Perforce client error: Connect to server failed;
-                #   check $P4PORT. TCP connect to perforce:1666 failed. No
-                #   such host is known."
-                if 'Perforce client error:' in output:
-                    raise ValueError
+                LOGGER.error(f'{self.name}: Error parsing syncstatus output: {output}')
 
-            except (ValueError, IndexError):
-                LOGGER.error(
-                    f"{self.name}: Could not retrieve changelists for "
-                    "project. Are the Source Control Settings correctly "
-                    "configured?")
-                return
+                if stderr:
+                    LOGGER.error(f'{self.name}: stderr: {stderr}')
+
+                LOGGER.error(f'{self.name}: Could not retrieve sync status. '
+                             'Are the source control settings correctly configured?')
 
-            if program_name.endswith("project"):
-                project_name = os.path.basename(
-                    os.path.dirname(CONFIG.UPROJECT_PATH.get_value(self.name)))
-                LOGGER.info(
-                    f"{self.name}: Project {project_name} "
-                    f"is on revision {current_changelist}")
-                for device in self.devices_sharing_workspace():
-                    device.project_changelist = current_changelist
-            elif program_name.endswith("engine"):
-                project_name = os.path.basename(
-                    os.path.dirname(CONFIG.UPROJECT_PATH.get_value(self.name)))
-                LOGGER.info(
-                    f"{self.name}: Engine used for project "
-                    f"{project_name} is on revision {current_changelist}")
                 for device in self.devices_sharing_workspace():
-                    device.engine_changelist = current_changelist
+                    if CONFIG.ENGINE_SYNC_METHOD.get_value() == EngineSyncMethod.Build_Engine.value:
+                        device.engine_changelist = '(unknown)'
+                    device.project_changelist = '(unknown)'
+
+                return
+
+            for device in self.devices_sharing_workspace():
+                dirty_str = '*' if syncstatus.get('dirty') else ''
+                if engine_cl := syncstatus.get('engine_cl'):
+                    device.engine_changelist = f'{engine_cl}{dirty_str}'
+                if project_cl := syncstatus.get('project_cl'):
+                    device.project_changelist = f'{project_cl}{dirty_str}'
+
         elif program_name == UnrealJobs.PackageGame.value:
             self.status = DeviceStatus.CLOSED
             percentstr = '100%'
@@ -3005,10 +2959,9 @@ class DeviceUnreal(Device):
 
         # request roles and changelists
         self._request_roles_file()
-        self._request_project_changelist_number()
+        self._request_syncstatus()
 
         if CONFIG.ENGINE_SYNC_METHOD.get_value() == EngineSyncMethod.Build_Engine.value:
-            self._request_engine_changelist_number()
             self._request_version_files()
 
     def transport_paths(self, device_recording):
diff --git a/Engine/Plugins/VirtualProduction/Switchboard/Source/Switchboard/switchboard/p4_utils.py b/Engine/Plugins/VirtualProduction/Switchboard/Source/Switchboard/switchboard/p4_utils.py
--- a/Engine/Plugins/VirtualProduction/Switchboard/Source/Switchboard/switchboard/p4_utils.py
+++ b/Engine/Plugins/VirtualProduction/Switchboard/Source/Switchboard/switchboard/p4_utils.py
@@ -168,16 +168,38 @@ def p4_from_localpath(localpath, workspaces, preferredClient):
 ##########
 
 
-def p4_variables() -> dict[str, str]:
+P4Env = dict[str, str]
+_cached_p4_envs: dict[pathlib.Path, P4Env] = {}
+
+def p4_get_env(path: pathlib.Path | None) -> P4Env:
+    global _cached_p4_envs
+
+    if path is None:
+        path = pathlib.Path.cwd()
+
+    env = _cached_p4_envs.get(path)
+    if env is None:
+        env = {}
+        logging.debug(f'Querying `p4 set` for path: {path}')
+        try:
+            output = subprocess.check_output(['p4', 'set', '-q'], cwd=path).decode()
+            for line in output.splitlines():
+                var_name, _, var_value = line.partition('=')
+                env[var_name] = var_value
+        except Exception as exc:
+            logging.exception(f'Querying `p4 set` failed with exception', exc_info=exc)
+
+        _cached_p4_envs[path] = env
+
+    return env
+
+def p4_get_var(var: str, path: pathlib.Path | None) -> str | None:
     '''
-    Returns Perforce variable (e.g. `P4CONFIG`, `p4 set`) key/value pairs.
+    Look up Perforce variable (e.g. `P4CONFIG`, `p4 set`) key/value pairs.
     '''
-    output = subprocess.check_output(['p4', 'set', '-q']).decode()
-    variables: dict[str, str] = {}
-    for line in output.splitlines():
-        var_name, _, var_value = line.partition('=')
-        variables[var_name] = var_value
-    return variables
+
+    env = p4_get_env(path)
+    return env.get(var)
 
 
 def p4(
@@ -185,6 +207,7 @@ def p4(
     cmd_opts: list[str],
     global_opts: Optional[list[str]] = None,
     *,
+    port: Optional[str] = None,
     user: Optional[str] = None,
     client: Optional[str] = None,
 ) -> subprocess.Popen:
@@ -192,9 +215,11 @@ def p4(
     args = ['p4', f'-zprog={meta_zprog}', f'-zversion={meta_zversion}',
             '-ztag', '-G', '-Qutf8']
 
-    if user is not None:
+    if port:
+        args.extend(['-p', port])
+    if user:
         args.extend(['-u', user])
-    if client is not None:
+    if client:
         args.extend(['-c', client])
 
     args.extend(global_opts or [])
@@ -216,11 +241,12 @@ def p4_get_records(
     include_info: bool = False,
     include_error: bool = False,
     log_unreturned: bool = True,
+    port: Optional[str] = None,
     user: Optional[str] = None,
     client: Optional[str] = None,
 ) -> list[dict]:
     results: list[dict] = []
-    proc = p4(cmd, opts, global_opts=global_opts, user=user, client=client)
+    proc = p4(cmd, opts, global_opts=global_opts, port=port, user=user, client=client)
     while True:
         try:
             record = marshal.load(proc.stdout)
@@ -234,7 +260,7 @@ def p4_get_records(
             continue
         elif (record[b'code'] == b'error') and not include_error:
             if log_unreturned:
-                logging.debug(f'p4_get_records: error: {record}')
+                logging.error(f'p4_get_records: error: {record}')
             continue
 
         results.append(record)
@@ -250,26 +276,44 @@ def p4_get_decoded_records(
     *,
     include_info: bool = False,
     include_error: bool = False,
+    port: Optional[str] = None,
     user: Optional[str] = None,
     client: Optional[str] = None,
 ) -> list[dict]:
     records = p4_get_records(
-        cmd, opts, global_opts, include_info=include_info,
-        include_error=include_error, user=user, client=client)
+        cmd, opts, global_opts, include_info=include_info, include_error=include_error,
+        port=port, user=user, client=client)
 
     return [{k.decode(): v.decode() if isinstance(v, bytes) else v
              for k, v in record.items()}
             for record in records]
 
 
+def p4_get_login(
+    *,
+    port: Optional[str] = None,
+    user: Optional[str] = None,
+    client: Optional[str] = None,
+) -> Optional[str]:
+    ''' Returns the logged in user, or None if there's no valid ticket. '''
+    records = p4_get_decoded_records('login', ['-s'],
+                                     port=port, user=user, client=client)
+
+    if records:
+        return records[0].get('User')
+    else:
+        return None
+
+
 def p4_get_client(
     get_client_name: str,
     *,
+    port: Optional[str] = None,
     user: Optional[str] = None,
     client: Optional[str] = None,
 ) -> Optional[dict]:
     records = p4_get_decoded_records('client', ['-o', get_client_name],
-                                     user=user, client=client)
+                                     port=port, user=user, client=client)
 
     if records:
         return records[0]
@@ -277,22 +321,44 @@ def p4_get_client(
         return None
 
 
+def p4_clients(
+    query_user: str,
+    *,
+    port: Optional[str] = None,
+    user: Optional[str] = None,
+    client: Optional[str] = None,
+):
+    records = p4_get_decoded_records('clients', ['-u', query_user],
+                                     port=port, user=user, client=client)
+
+    clients: dict[str, dict] = {}
+    for record in records:
+        if record_client := record.get('client'):
+            clients[record_client] = record
+        else:
+            logging.error(f'p4_clients: record missing "client" field: {record}')
+
+    return clients
+
+
 def p4_changes(
     pathspecs: list[str],
     *,
     limit: int = 10,
+    port: Optional[str] = None,
     user: Optional[str] = None,
     client: Optional[str] = None,
 ) -> list[dict]:
     limit = min(limit, 100)
     opts = ['-s', 'submitted', '-t', '-L', f'-m{limit}', *pathspecs]
-    return p4_get_records('changes', opts, user=user, client=client)
+    return p4_get_records('changes', opts, port=port, user=user, client=client)
 
 
 def p4_fstat(
     pathspecs: list[str],
     *,
     limit: int = 0,
+    port: Optional[str] = None,
     user: Optional[str] = None,
     client: Optional[str] = None,
 ) -> list[dict]:
@@ -304,17 +370,18 @@ def p4_fstat(
     opts.extend(pathspecs)
 
     return p4_get_records('fstat', opts, log_unreturned=False,
-                          user=user, client=client)
+                          port=port, user=user, client=client)
 
 
 def p4_have(
     pathspecs: list[str],
     *,
+    port: Optional[str] = None,
     user: Optional[str] = None,
     client: Optional[str] = None,
 ) -> list[dict]:
     records = p4_get_decoded_records('have', pathspecs, include_error=True,
-                                     user=user, client=client)
+                                     port=port, user=user, client=client)
 
     return records
 
@@ -343,6 +410,7 @@ class P4PrintResult:
 def p4_print(
     pathspecs: list[str],
     *,
+    port: Optional[str] = None,
     user: Optional[str] = None,
     client: Optional[str] = None,
 ) -> list[P4PrintResult]:
@@ -352,8 +420,8 @@ def p4_print(
     fileContents will be `None`.
     '''
     results: list[P4PrintResult] = []
-    records = p4_get_records('print', pathspecs, include_info=True,
-                             include_error=True, user=user, client=client)
+    records = p4_get_records('print', pathspecs, include_info=True, include_error=True,
+                             port=port, user=user, client=client)
 
     # For each file, a stat record is returned, then zero or more text records.
     # Cases where there's no text record include errors or deleted files.
@@ -383,6 +451,7 @@ def p4_latest_code_change(
     *,
     in_range: Optional[str] = None,
     exts: Optional[list[str]] = None,
+    port: Optional[str] = None,
     user: Optional[str] = None,
     client: Optional[str] = None,
 ) -> Optional[int]:
@@ -402,7 +471,7 @@ def p4_latest_code_change(
                                for ext in exts])
 
     code_changes = p4_changes(code_paths, limit=1,
-                              user=user, client=client)
+                              port=port, user=user, client=client)
 
     if len(code_changes) > 0:
         return max(int(x[b'change']) for x in code_changes)
diff --git a/Engine/Plugins/VirtualProduction/Switchboard/Source/Switchboard/switchboard/ugs_utils.py b/Engine/Plugins/VirtualProduction/Switchboard/Source/Switchboard/switchboard/ugs_utils.py
--- a/Engine/Plugins/VirtualProduction/Switchboard/Source/Switchboard/switchboard/ugs_utils.py
+++ b/Engine/Plugins/VirtualProduction/Switchboard/Source/Switchboard/switchboard/ugs_utils.py
@@ -319,7 +319,7 @@ def sync(
     ugs_bin_dir: Optional[pathlib.Path] = None,
     user: Optional[str] = None,
     client: Optional[str] = None
-) -> Optional[int]:
+) -> int:
 
     sync_args = ['sync']
     if sync_cl:
@@ -509,24 +509,30 @@ class IniParser:
 def parse_depot_ugs_configs(
     engine_dir: Optional[Union[str, pathlib.Path]],
     project_dir: Optional[Union[str, pathlib.Path]],
-    p4user: Optional[str] = None,
-    p4client: Optional[str] = None,
+    *,
+    port: Optional[str] = None,
+    user: Optional[str] = None,
+    client: Optional[str] = None,
 ) -> IniParser:
     '''
     Given local or depot engine and/or project paths, search the Perforce depot
     for UnrealGameSync.ini configs at known locations, coalesce their contents,
     and return an `IniParser` populated with the results.
     '''
+    parser = IniParser()
 
-    p4opts = {'user': p4user, 'client': p4client}
-
-    logging.debug('Parsing UnrealGameSync configuration from depot')
+    p4kwargs = {'port': port, 'user': user, 'client': client}
 
     # Might append depot wildcards to host paths and end up with mixed syntax
     # here, but the fstat separates them into normalized depot/host paths.
     mixed_config_paths = get_depot_config_paths(engine_dir, project_dir)
+    if not mixed_config_paths:
+        logging.debug('No paths to check for UnrealGameSync configuration')
+        return parser
+
+    logging.debug('Parsing UnrealGameSync configuration from depot')
 
-    config_fstats = p4_utils.p4_fstat(mixed_config_paths, **p4opts)
+    config_fstats = p4_utils.p4_fstat(mixed_config_paths, **p4kwargs)
 
     config_contents_map: dict[str, Optional[str]] = {}
 
@@ -560,7 +566,7 @@ def parse_depot_ugs_configs(
     depot_print_paths = [k for k, v in config_contents_map.items()
                          if v is None]  # configs we didn't read locally
 
-    config_prints = p4_utils.p4_print(depot_print_paths, **p4opts)
+    config_prints = p4_utils.p4_print(depot_print_paths, **p4kwargs)
 
     for depot_path, record in zip(depot_print_paths, config_prints):
         if record.is_valid:
@@ -570,7 +576,6 @@ def parse_depot_ugs_configs(
         else:
             logging.warning(f'Invalid P4PrintResult for {depot_path}')
 
-    parser = IniParser()
     for path, contents in config_contents_map.items():
         if contents is None:
             logging.warning(f'No contents for {path}')
