Switchboard: Include rsync daemon in job object to ensure it doesn't outlive Switchboard.

Also centralizes other ctypes Win32 glue code.

diff --git a/Engine/Plugins/VirtualProduction/Switchboard/Source/Switchboard/switchboard/credential_store.py b/Engine/Plugins/VirtualProduction/Switchboard/Source/Switchboard/switchboard/credential_store.py
--- a/Engine/Plugins/VirtualProduction/Switchboard/Source/Switchboard/switchboard/credential_store.py
+++ b/Engine/Plugins/VirtualProduction/Switchboard/Source/Switchboard/switchboard/credential_store.py
@@ -133,113 +133,15 @@ if sys.platform.startswith('win'):
     import ctypes
     import ctypes.wintypes
 
-    CRED_PERSIST_LOCAL_MACHINE = 2
-    CRED_TYPE_GENERIC = 1
-    ERROR_NOT_FOUND = 1168
-
-    class CREDENTIALW(ctypes.Structure):
-        _fields_ = [
-            ('Flags', ctypes.wintypes.DWORD),
-            ('Type', ctypes.wintypes.DWORD),
-            ('TargetName', ctypes.wintypes.LPWSTR),
-            ('Comment', ctypes.wintypes.LPWSTR),
-            ('LastWritten', ctypes.wintypes.FILETIME),
-            ('CredentialBlobSize', ctypes.wintypes.DWORD),
-            ('CredentialBlob', ctypes.wintypes.LPBYTE),
-            ('Persist', ctypes.wintypes.DWORD),
-            ('AttributeCount', ctypes.wintypes.DWORD),
-            ('Attributes', ctypes.c_void_p),
-            ('TargetAlias', ctypes.wintypes.LPWSTR),
-            ('UserName', ctypes.wintypes.LPWSTR),
-        ]
+    from .switchboard_ctypes import (CREDENTIALW, CredDeleteW, CredFree, CredReadW, CredWriteW,
+                                     CRED_PERSIST_LOCAL_MACHINE, CRED_TYPE_GENERIC)
+
 
     class CredentialStoreWindows(CredentialStoreGeneric):
         def __init__(self):
             super().__init__()
-
             del self._tempstore
 
-            self._advapi32 = ctypes.WinDLL('Advapi32')
-
-            self._CredFree_prototype = ctypes.WINFUNCTYPE(
-                None,  # return type
-
-                ctypes.c_void_p,  # Buffer
-            )
-
-            self._CredFree = self._CredFree_prototype(
-                ('CredFree', self._advapi32),
-                (
-                    (1, 'Buffer'),
-                )
-            )
-
-            self._CredDeleteW_prototype = ctypes.WINFUNCTYPE(
-                ctypes.wintypes.BOOL,  # return type
-
-                ctypes.wintypes.LPCWSTR,  # TargetName
-                ctypes.wintypes.DWORD,  # Type
-                ctypes.wintypes.DWORD,  # Flags (reserved)
-            )
-
-            self._CredDeleteW = self._CredDeleteW_prototype(
-                ('CredDeleteW', self._advapi32),
-                (
-                    (1, 'TargetName'),
-                    (1, 'Type'),
-                    (1, 'Flags', 0),
-                )
-            )
-
-            self._CredReadW_prototype = ctypes.WINFUNCTYPE(
-                ctypes.wintypes.BOOL,  # return type
-
-                ctypes.wintypes.LPWSTR,  # TargetName
-                ctypes.wintypes.DWORD,  # Type
-                ctypes.wintypes.DWORD,  # Flags (reserved)
-                ctypes.POINTER(ctypes.POINTER(CREDENTIALW)),  # OutCredential
-
-                use_last_error=True,
-            )
-            self._CredReadW = self._CredReadW_prototype(
-                ("CredReadW", self._advapi32),
-                (
-                    (1, 'TargetName'),
-                    (1, 'Type'),
-                    (1, 'Flags', 0),
-                    (3, 'Credential'),
-                )
-            )
-
-            self._CredWriteW_prototype = ctypes.WINFUNCTYPE(
-                ctypes.wintypes.BOOL,  # return type
-
-                ctypes.POINTER(CREDENTIALW),  # Credential
-                ctypes.wintypes.DWORD,  # Flags
-
-                use_last_error=True,
-            )
-            self._CredWriteW = self._CredWriteW_prototype(
-                ("CredWriteW", self._advapi32),
-                (
-                    (1, 'Credential'),
-                    (1, 'Flags', 0),
-                )
-            )
-
-            def Cred_errcheck(result, func, args):
-                if not result:
-                    last_err = ctypes.get_last_error()
-                    if last_err == ERROR_NOT_FOUND:
-                        raise KeyError()
-                    else:
-                        raise ctypes.WinError(last_err)
-
-                return args
-
-            self._CredReadW.errcheck = Cred_errcheck
-            self._CredWriteW.errcheck = Cred_errcheck
-
         @classmethod
         def encrypted_at_rest(cls) -> bool:
             return True
@@ -248,7 +150,7 @@ if sys.platform.startswith('win'):
             need_credfree = False
             cred = ctypes.pointer(CREDENTIALW())
             try:
-                self._CredReadW(
+                CredReadW(
                     TargetName=f'Switchboard_{key}',
                     Type=CRED_TYPE_GENERIC,
                     Flags=0,
@@ -270,7 +172,7 @@ if sys.platform.startswith('win'):
                 return None
             finally:
                 if need_credfree:
-                    self._CredFree(cred)
+                    CredFree(cred)
 
         def set(
             self,
@@ -279,7 +181,7 @@ if sys.platform.startswith('win'):
         ) -> None:
             target_name = f'Switchboard_{key}'
             if not val:
-                self._CredDeleteW(target_name, CRED_TYPE_GENERIC)
+                CredDeleteW(target_name, CRED_TYPE_GENERIC)
                 return
 
             cred = CREDENTIALW()
@@ -297,7 +199,7 @@ if sys.platform.startswith('win'):
             cred.CredentialBlob = ((ctypes.wintypes.BYTE * blob_byte_len)
                                    .from_buffer(blob_bytearr))
 
-            self._CredWriteW(ctypes.byref(cred), 0)
+            CredWriteW(ctypes.byref(cred), 0)
 
 
 CREDENTIAL_STORE = CredentialStore.create()
diff --git a/Engine/Plugins/VirtualProduction/Switchboard/Source/Switchboard/switchboard/switchboard_application.py b/Engine/Plugins/VirtualProduction/Switchboard/Source/Switchboard/switchboard/switchboard_application.py
--- a/Engine/Plugins/VirtualProduction/Switchboard/Source/Switchboard/switchboard/switchboard_application.py
+++ b/Engine/Plugins/VirtualProduction/Switchboard/Source/Switchboard/switchboard/switchboard_application.py
@@ -10,7 +10,6 @@ import subprocess
 import sys
 import tempfile
 import threading
-
 from typing import NamedTuple, Optional, Union
 
 import pythonosc.dispatcher
@@ -18,6 +17,7 @@ import pythonosc.osc_server
 
 from .config import CONFIG, SETTINGS
 from .switchboard_logging import LOGGER
+from . import switchboard_ctypes as sb_ctypes
 from . import switchboard_utils as sb_utils
 
 
@@ -347,6 +347,12 @@ hosts allow = {allowed_addrs}
         self._monitor_thread: Optional[threading.Thread] = None
         self._monitor_recoveries = 0
 
+        self._cleanup_helper = None
+        try:
+            self._cleanup_helper = sb_ctypes.CloseSubprocessOnExitHelper()
+        except Exception as exc:
+            LOGGER.warning('RsyncServer: Error creating cleanup helper')
+
     def is_running(self) -> bool:
         return (self.process is not None) and (self.process.poll() is None)
 
@@ -438,7 +444,6 @@ hosts allow = {allowed_addrs}
         LOGGER.debug(f'RsyncServer.launch: {" ".join(args)}')
 
         try:
-            # TODO: Look into AssignProcessToJobObject to ensure cleanup?
             self.process = subprocess.Popen(
                 args,
                 stdin=subprocess.DEVNULL,
@@ -450,6 +455,12 @@ hosts allow = {allowed_addrs}
                          f'Failed with exception: {exc}', exc_info=exc)
             return False
 
+        try:
+            if self._cleanup_helper:
+                self._cleanup_helper.add_process(self.process.pid)
+        except Exception as exc:
+            LOGGER.warning('RsyncServer.launch: Error registering process for cleanup', exc_info=exc)
+
         self._shutdown_event.clear()
         self._monitor_thread = threading.Thread(
             target=self._monitor)
diff --git a/Engine/Plugins/VirtualProduction/Switchboard/Source/Switchboard/switchboard/switchboard_ctypes.py b/Engine/Plugins/VirtualProduction/Switchboard/Source/Switchboard/switchboard/switchboard_ctypes.py
new file mode 100644
--- /dev/null
+++ b/Engine/Plugins/VirtualProduction/Switchboard/Source/Switchboard/switchboard/switchboard_ctypes.py
@@ -0,0 +1,269 @@
+# Copyright Epic Games, Inc. All Rights Reserved.
+
+from __future__ import annotations
+import ctypes
+import sys
+
+
+if sys.platform.startswith('win'):
+    import ctypes.wintypes
+
+    advapi32 = ctypes.WinDLL('Advapi32')
+    kernel32 = ctypes.WinDLL('kernel32')
+
+    ERROR_NOT_FOUND = 1168
+
+    def falsey_errcheck(result, func, args):
+        if not result:
+            last_err = ctypes.get_last_error()
+            raise ctypes.WinError(last_err)
+
+        return args
+
+    ####################
+    # Credentials
+
+    CRED_PERSIST_LOCAL_MACHINE = 2
+    CRED_TYPE_GENERIC = 1
+
+    class CREDENTIALW(ctypes.Structure):
+        _fields_ = [
+            ('Flags', ctypes.wintypes.DWORD),
+            ('Type', ctypes.wintypes.DWORD),
+            ('TargetName', ctypes.wintypes.LPWSTR),
+            ('Comment', ctypes.wintypes.LPWSTR),
+            ('LastWritten', ctypes.wintypes.FILETIME),
+            ('CredentialBlobSize', ctypes.wintypes.DWORD),
+            ('CredentialBlob', ctypes.wintypes.LPBYTE),
+            ('Persist', ctypes.wintypes.DWORD),
+            ('AttributeCount', ctypes.wintypes.DWORD),
+            ('Attributes', ctypes.c_void_p),
+            ('TargetAlias', ctypes.wintypes.LPWSTR),
+            ('UserName', ctypes.wintypes.LPWSTR),
+        ]
+
+    _CredFree_prototype = ctypes.WINFUNCTYPE(
+        None,  # return type
+
+        ctypes.c_void_p,  # Buffer
+    )
+    CredFree = _CredFree_prototype(
+        ('CredFree', advapi32),
+        (
+            (1, 'Buffer'),
+        )
+    )
+
+    _CredDeleteW_prototype = ctypes.WINFUNCTYPE(
+        ctypes.wintypes.BOOL,  # return type
+
+        ctypes.wintypes.LPCWSTR,  # TargetName
+        ctypes.wintypes.DWORD,    # Type
+        ctypes.wintypes.DWORD,    # Flags (reserved)
+    )
+    CredDeleteW = _CredDeleteW_prototype(
+        ('CredDeleteW', advapi32),
+        (
+            (1, 'TargetName'),
+            (1, 'Type'),
+            (1, 'Flags', 0),
+        )
+    )
+
+    _CredReadW_prototype = ctypes.WINFUNCTYPE(
+        ctypes.wintypes.BOOL,  # return type
+
+        ctypes.wintypes.LPWSTR,  # TargetName
+        ctypes.wintypes.DWORD,   # Type
+        ctypes.wintypes.DWORD,   # Flags (reserved)
+        ctypes.POINTER(ctypes.POINTER(CREDENTIALW)),  # OutCredential
+
+        use_last_error=True,
+    )
+    CredReadW = _CredReadW_prototype(
+        ('CredReadW', advapi32),
+        (
+            (1, 'TargetName'),
+            (1, 'Type'),
+            (1, 'Flags', 0),
+            (3, 'Credential'),
+        )
+    )
+
+    _CredWriteW_prototype = ctypes.WINFUNCTYPE(
+        ctypes.wintypes.BOOL,  # return type
+
+        ctypes.POINTER(CREDENTIALW),  # Credential
+        ctypes.wintypes.DWORD,        # Flags
+
+        use_last_error=True,
+    )
+    CredWriteW = _CredWriteW_prototype(
+        ('CredWriteW', advapi32),
+        (
+            (1, 'Credential'),
+            (1, 'Flags', 0),
+        )
+    )
+
+    def Cred_errcheck(result, func, args):
+        if not result:
+            last_err = ctypes.get_last_error()
+            if last_err == ERROR_NOT_FOUND:
+                raise KeyError()
+            else:
+                raise ctypes.WinError(last_err)
+
+        return args
+
+    CredReadW.errcheck = Cred_errcheck
+    CredWriteW.errcheck = Cred_errcheck
+
+    ####################
+    # Job objects
+
+    ULONG_PTR = SIZE_T = ctypes.wintypes.WPARAM
+
+    PROCESS_TERMINATE = 0x0001
+    PROCESS_SET_QUOTA = 0x0100
+    JOB_OBJECT_LIMIT_KILL_ON_JOB_CLOSE = 0x00002000
+    JobObjectExtendedLimitInformation = 9
+
+    class JOBOBJECT_BASIC_LIMIT_INFORMATION(ctypes.Structure):
+        _fields_ = [
+            ('PerProcessUserTimeLimit', ctypes.wintypes.LARGE_INTEGER),
+            ('PerJobUserTimeLimit', ctypes.wintypes.LARGE_INTEGER),
+            ('LimitFlags', ctypes.wintypes.DWORD),
+            ('MinimumWorkingSetSize', SIZE_T),
+            ('MaximumWorkingSetSize', SIZE_T),
+            ('ActiveProcessLimit', ctypes.wintypes.DWORD),
+            ('Affinity', ULONG_PTR),
+            ('PriorityClass', ctypes.wintypes.DWORD),
+            ('SchedulingClass', ctypes.wintypes.DWORD),
+        ]
+
+    class IO_COUNTERS(ctypes.Structure):
+        _fields_ = [
+            ('ReadOperationCount', ctypes.wintypes.ULARGE_INTEGER),
+            ('WriteOperationCount', ctypes.wintypes.ULARGE_INTEGER),
+            ('OtherOperationCount', ctypes.wintypes.ULARGE_INTEGER),
+            ('ReadTransferCount', ctypes.wintypes.ULARGE_INTEGER),
+            ('WriteTransferCount', ctypes.wintypes.ULARGE_INTEGER),
+            ('OtherTransferCount', ctypes.wintypes.ULARGE_INTEGER),
+        ]
+
+    class JOBOBJECT_EXTENDED_LIMIT_INFORMATION(ctypes.Structure):
+        _fields_ = [
+            ('BasicLimitInformation', JOBOBJECT_BASIC_LIMIT_INFORMATION),
+            ('IoInfo', IO_COUNTERS),
+            ('ProcessMemoryLimit', SIZE_T),
+            ('JobMemoryLimit', SIZE_T),
+            ('PeakProcessMemoryUsed', SIZE_T),
+            ('PeakJobMemoryUsed', SIZE_T),
+        ]
+
+    _CreateJobObjectW_prototype = ctypes.WINFUNCTYPE(
+        ctypes.wintypes.HANDLE,  # return type
+
+        ctypes.wintypes.LPVOID,   # JobAttributes (LPSECURITY_ATTRIBUTES*)
+        ctypes.wintypes.LPCWSTR,  # Name
+
+        use_last_error=True,
+    )
+    CreateJobObjectW = _CreateJobObjectW_prototype(
+        ('CreateJobObjectW', kernel32),
+        (
+            (1, 'JobAttributes', 0),
+            (1, 'Name'),
+        )
+    )
+    CreateJobObjectW.errcheck = falsey_errcheck
+
+    _QueryInformationJobObject_prototype = ctypes.WINFUNCTYPE(
+        ctypes.wintypes.BOOL,  # return type
+
+        ctypes.wintypes.HANDLE,   # hJob
+        ctypes.wintypes.INT,      # JobObjectInformationClass (enum JOBOBJECTINFOCLASS)
+        ctypes.wintypes.LPVOID,   # lpJobObjectInformation
+        ctypes.wintypes.DWORD,    # cbJobObjectInformationLength
+        ctypes.wintypes.LPDWORD,  # lpReturnLength
+
+        use_last_error=True,
+    )
+    QueryInformationJobObject = _QueryInformationJobObject_prototype(
+        ('QueryInformationJobObject', kernel32),
+        (
+            (1, 'hJob'),
+            (1, 'JobObjectInformationClass'),
+            (3, 'lpJobObjectInformation'),
+            (1, 'cbJobObjectInformationLength'),
+            (3, 'lpReturnLength'),
+        )
+    )
+    QueryInformationJobObject.errcheck = falsey_errcheck
+
+    _SetInformationJobObject_prototype = ctypes.WINFUNCTYPE(
+        ctypes.wintypes.BOOL,  # return type
+
+        ctypes.wintypes.HANDLE,  # hJob
+        ctypes.wintypes.INT,     # JobObjectInformationClass (enum JOBOBJECTINFOCLASS)
+        ctypes.wintypes.LPVOID,  # lpJobObjectInformation
+        ctypes.wintypes.DWORD,   # cbJobObjectInformationLength
+
+        use_last_error=True,
+    )
+    SetInformationJobObject = _SetInformationJobObject_prototype(
+        ('SetInformationJobObject', kernel32),
+        (
+            (1, 'hJob'),
+            (1, 'JobObjectInformationClass'),
+            (1, 'lpJobObjectInformation'),
+            (1, 'cbJobObjectInformationLength'),
+        )
+    )
+    SetInformationJobObject.errcheck = falsey_errcheck
+
+
+class CloseSubprocessOnExitHelper:
+    ''' Helper to ensure a child process exits if this parent process does. '''
+
+    def __init__(self):
+        if sys.platform.startswith('win'):
+            # Create job object and configure it to ensure termination
+            self._job_object = CreateJobObjectW(None, None)
+
+            ext_limit_info = JOBOBJECT_EXTENDED_LIMIT_INFORMATION()
+            QueryInformationJobObject(
+                self._job_object,
+                JobObjectExtendedLimitInformation,
+                ctypes.byref(ext_limit_info),
+                ctypes.sizeof(ext_limit_info),
+                None,
+            )
+
+            ext_limit_info.BasicLimitInformation.LimitFlags |= JOB_OBJECT_LIMIT_KILL_ON_JOB_CLOSE
+
+            SetInformationJobObject(
+                self._job_object,
+                JobObjectExtendedLimitInformation,
+                ctypes.byref(ext_limit_info),
+                ctypes.sizeof(ext_limit_info),
+            )
+
+    def add_process(self, pid):
+        ''' Bind the process with the specified PID to terminate when this process does. '''
+
+        if sys.platform.startswith('win'):
+            # Get native handle to process
+            desired_access = ctypes.wintypes.DWORD(PROCESS_TERMINATE | PROCESS_SET_QUOTA)
+            inherit_handle = ctypes.wintypes.BOOL(False)
+            process_id = ctypes.wintypes.DWORD(pid)
+            process_handle = kernel32.OpenProcess(desired_access, inherit_handle, process_id)
+
+            if not process_handle:
+                raise ctypes.WinError(ctypes.get_last_error())
+
+            # Assign the process to our job object
+            kernel32.AssignProcessToJobObject(self._job_object, process_handle)
+
+            kernel32.CloseHandle(process_handle)
