Switchboard: Move local P4 changelist refresh off main thread.

diff --git a/Engine/Plugins/VirtualProduction/Switchboard/Source/Switchboard/switchboard/switchboard_dialog.py b/Engine/Plugins/VirtualProduction/Switchboard/Source/Switchboard/switchboard/switchboard_dialog.py
--- a/Engine/Plugins/VirtualProduction/Switchboard/Source/Switchboard/switchboard/switchboard_dialog.py
+++ b/Engine/Plugins/VirtualProduction/Switchboard/Source/Switchboard/switchboard/switchboard_dialog.py
@@ -2053,11 +2053,34 @@ class SwitchboardDialog(QtCore.QObject):
         if self.logger_autoscroll:
             self.logger_scroll_to_end()
 
+
+    class ChangelistRunnable(QtCore.QRunnable):
+        ''' Retrieves a list of recent changelists for the given path(s) on the thread pool. '''
+
+        class Signals(QtCore.QObject):
+            result = QtCore.Signal(list) # list[tuple[int, str]]
+
+        def __init__(self, paths: list[str], client: str | None = None):
+            super().__init__()
+            self.paths = paths
+            self.client = client
+            self.signals = self.Signals()
+
+        def run(self):
+            cl_desc_pairs = p4_utils.p4_latest_changelists(self.paths, client=self.client)
+            self.signals.result.emit(cl_desc_pairs)
+
+    _engine_cl_runnable: ChangelistRunnable | None = None
+    _project_cl_runnable: ChangelistRunnable | None = None
+
     # Update UI with latest CLs
     def p4_refresh_project_cl(self):
         if not CONFIG.P4_ENABLED.get_value():
             return
 
+        if self._project_cl_runnable:
+            return # Already refreshing.
+
         sync_method = CONFIG.ENGINE_SYNC_METHOD.get_value()
 
         # If we're syncing 'Precompiled Binaries', then that implies that we should be using UGS:
@@ -2074,9 +2097,6 @@ class SwitchboardDialog(QtCore.QObject):
             self.window.refresh_project_cl_button.setEnabled(False)
             return
 
-        changelists: Optional[list[int]] = None
-        descriptions: Optional[list[str]] = None
-
         client = CONFIG.SOURCE_CONTROL_WORKSPACE.get_value()
         project_client = CONFIG.PROJECT_WORKSPACE.get_value()
 
@@ -2085,41 +2105,51 @@ class SwitchboardDialog(QtCore.QObject):
 
         if use_ugs:
             LOGGER.info("Using UnrealGameSync to refresh project changelists.")
-            changelists = ugs_utils.latest_chagelists(Path(CONFIG.UPROJECT_PATH.get_value()), client=client)
-            if not changelists:
-                LOGGER.error("UnrealGameSync failed to get the project's latest changelists. Falling back to using p4 commands directly.")
+            if changelists := ugs_utils.latest_chagelists(Path(CONFIG.UPROJECT_PATH.get_value()), client=client):
+                self._on_project_changelists_refreshed(list([(cl, "") for cl in changelists]))
+                return
+            else:
+                LOGGER.error("UnrealGameSync failed to get the project's latest changelists. "
+                             "Falling back to using p4 commands directly.")
 
-        if not changelists:
+        LOGGER.info("Refreshing p4 project changelists")
+        paths = [f'{CONFIG.P4_PROJECT_PATH.get_value()}/...']
 
-            LOGGER.info("Refreshing p4 project changelists")
-            paths = [f'{CONFIG.P4_PROJECT_PATH.get_value()}/...']
+        if DeviceUnreal.ugs_config:
+            if addpaths := DeviceUnreal.ugs_config.try_get(
+                'Perforce', 'AdditionalPathsToSync'
+            ):
+                paths.extend([f'//{client}{x}' for x in addpaths])
 
-            if DeviceUnreal.ugs_config:
-                if addpaths := DeviceUnreal.ugs_config.try_get(
-                    'Perforce', 'AdditionalPathsToSync'
-                ):
-                    paths.extend([f'//{client}{x}' for x in addpaths])
+        self.window.project_cl_combo_box.clear()
+        self.window.project_cl_combo_box.addItem('...')
+        self._project_cl_runnable = SwitchboardDialog.ChangelistRunnable(paths, client=client)
+        self._project_cl_runnable.signals.result.connect(self._on_project_changelists_refreshed)
+        QtCore.QThreadPool.globalInstance().start(self._project_cl_runnable)
 
-            cl_descs = p4_utils.p4_latest_changelists(paths, client=client)
-            changelists = [str(cl[0]) for cl in cl_descs]
-            descriptions = [cl[1] for cl in cl_descs]
+    @QtCore.Slot()
+    def _on_project_changelists_refreshed(self, cl_desc_pairs: list[tuple[int, str]]):
+        if self._project_cl_runnable:
+            self._project_cl_runnable.signals.result.disconnect(self._on_project_changelists_refreshed)
+            self._project_cl_runnable = None
 
         self.window.project_cl_combo_box.clear()
 
-        if changelists:
-            self.window.project_cl_combo_box.addItems(changelists)
-            self.window.project_cl_combo_box.setCurrentIndex(0)
-
-        if descriptions:
-            for (idx, desc) in zip(count(), descriptions):
-                self.window.project_cl_combo_box.setItemData(
-                    idx, desc, QtCore.Qt.ItemDataRole.ToolTipRole)
+        for (idx, (cl, desc)) in zip(count(), cl_desc_pairs):
+            self.window.project_cl_combo_box.addItem(str(cl))
+            self.window.project_cl_combo_box.setItemData(idx, desc, QtCore.Qt.ItemDataRole.ToolTipRole)
 
+        self.window.project_cl_combo_box.setCurrentIndex(0)
         self.window.project_cl_combo_box.addItem(EMPTY_SYNC_ENTRY)
+        self.window.project_cl_combo_box.setCurrentText(self._project_changelist)
 
     def p4_refresh_engine_cl(self):
         if not CONFIG.P4_ENABLED.get_value():
             return
+
+        if self._engine_cl_runnable:
+            return # Already refreshing.
+
         self.window.engine_cl_combo_box.clear()
         # if engine is built from source, refresh the engine cl dropdown
         if CONFIG.ENGINE_SYNC_METHOD.get_value() == EngineSyncMethod.Build_Engine.value:
@@ -2132,20 +2162,13 @@ class SwitchboardDialog(QtCore.QObject):
 
             engine_p4_path = CONFIG.P4_ENGINE_PATH.get_value()
             if engine_p4_path:
+                self.window.engine_cl_combo_box.clear()
+                self.window.engine_cl_combo_box.addItem('...')
+
                 client = CONFIG.SOURCE_CONTROL_WORKSPACE.get_value()
-                cl_descs = p4_utils.p4_latest_changelists(
-                    engine_p4_path+'/...', client=client)
-                changelists = [str(cl[0]) for cl in cl_descs]
-                descriptions = [cl[1] for cl in cl_descs]
-
-                if changelists:
-                    self.window.engine_cl_combo_box.addItems(changelists)
-                    self.window.engine_cl_combo_box.setCurrentIndex(0)
-
-                if descriptions:
-                    for (idx, desc) in zip(count(), descriptions):
-                        self.window.engine_cl_combo_box.setItemData(
-                            idx, desc, QtCore.Qt.ItemDataRole.ToolTipRole)
+                self._engine_cl_runnable = SwitchboardDialog.ChangelistRunnable([engine_p4_path+'/...'], client=client)
+                self._engine_cl_runnable.signals.result.connect(self._on_engine_changelists_refreshed)
+                QtCore.QThreadPool.globalInstance().start(self._engine_cl_runnable)
             else:
                 LOGGER.warning('"Build Engine" is enabled in the settings but the engine does not seem to be under perforce control.')
                 LOGGER.warning("Please check your perforce settings.")
@@ -2158,7 +2181,22 @@ class SwitchboardDialog(QtCore.QObject):
             self.window.engine_cl_combo_box.setToolTip(tool_tip)
             self.window.refresh_engine_cl_button.setEnabled(False)
             self.window.refresh_engine_cl_button.setToolTip(tool_tip)
+
+    @QtCore.Slot()
+    def _on_engine_changelists_refreshed(self, cl_desc_pairs: list[tuple[int, str]]):
+        if self._engine_cl_runnable:
+            self._engine_cl_runnable.signals.result.disconnect(self._on_engine_changelists_refreshed)
+            self._engine_cl_runnable = None
+
+        self.window.engine_cl_combo_box.clear()
+
+        for (idx, (cl, desc)) in zip(count(), cl_desc_pairs):
+            self.window.engine_cl_combo_box.addItem(str(cl))
+            self.window.engine_cl_combo_box.setItemData(idx, desc, QtCore.Qt.ItemDataRole.ToolTipRole)
+
+        self.window.engine_cl_combo_box.setCurrentIndex(0)
         self.window.engine_cl_combo_box.addItem(EMPTY_SYNC_ENTRY)
+        self.window.engine_cl_combo_box.setCurrentText(self._engine_changelist)
 
     def refresh_levels(self, levels: Optional[List[str]] = None):
 
