Switchboard: Allow population / refresh of the asset cache to be interrupted.

diff --git a/Engine/Plugins/VirtualProduction/Switchboard/Source/Switchboard/switchboard/config.py b/Engine/Plugins/VirtualProduction/Switchboard/Source/Switchboard/switchboard/config.py
--- a/Engine/Plugins/VirtualProduction/Switchboard/Source/Switchboard/switchboard/config.py
+++ b/Engine/Plugins/VirtualProduction/Switchboard/Source/Switchboard/switchboard/config.py
@@ -3268,7 +3268,9 @@ class Config(object):
 
     def get_levels(self) -> list[str]:
         ''' Returns a list with all *.map files from the cache '''
-        project = SBCache().query_or_create_project(self.UPROJECT_PATH.get_value())
+        project = SBCache().query_project(self.UPROJECT_PATH.get_value())
+        if project is None:
+            return []
         maps = SBCache().query_maps(project)
         return [map.gamepath for map in maps]
 
diff --git a/Engine/Plugins/VirtualProduction/Switchboard/Source/Switchboard/switchboard/devices/ndisplay/plugin_ndisplay.py b/Engine/Plugins/VirtualProduction/Switchboard/Source/Switchboard/switchboard/devices/ndisplay/plugin_ndisplay.py
--- a/Engine/Plugins/VirtualProduction/Switchboard/Source/Switchboard/switchboard/devices/ndisplay/plugin_ndisplay.py
+++ b/Engine/Plugins/VirtualProduction/Switchboard/Source/Switchboard/switchboard/devices/ndisplay/plugin_ndisplay.py
@@ -188,7 +188,10 @@ class AddnDisplayDialog(AddDeviceDialog):
 
         self.cbConfigs.clear()
 
-        project = SBCache().query_or_create_project(CONFIG.UPROJECT_PATH.get_value())
+        project = SBCache().query_project(CONFIG.UPROJECT_PATH.get_value())
+        if project is None:
+            return
+
         assets = SBCache().query_assets_by_classname(project=project, classnames=DeviceUnreal.NDISPLAY_CLASS_NAMES)
 
         for asset in assets:
diff --git a/Engine/Plugins/VirtualProduction/Switchboard/Source/Switchboard/switchboard/devices/unreal/plugin_unreal.py b/Engine/Plugins/VirtualProduction/Switchboard/Source/Switchboard/switchboard/devices/unreal/plugin_unreal.py
--- a/Engine/Plugins/VirtualProduction/Switchboard/Source/Switchboard/switchboard/devices/unreal/plugin_unreal.py
+++ b/Engine/Plugins/VirtualProduction/Switchboard/Source/Switchboard/switchboard/devices/unreal/plugin_unreal.py
@@ -16,6 +16,7 @@ import time
 import uuid
 import weakref
 
+from abc import abstractmethod
 from collections import OrderedDict
 from datetime import datetime
 from enum import Enum
@@ -23,7 +24,7 @@ from functools import wraps
 from ipaddress import IPv4Address
 from itertools import count
 from pathlib import PurePosixPath
-from typing import Callable, Dict, Generator, Optional, Union
+from typing import Callable, Dict, Generator, Iterable, Optional, Union
 
 from PySide6 import QtCore, QtGui, QtWidgets
 
@@ -344,13 +345,14 @@ class ProgramStartQueue:
         return sum(len(program_list) for program_list in self.rd_running_progs_by_name.values())
 
 
-class LiveLinkPresetSetting(Setting):
-    ''' Container of the LiveLink Preset setting
-    Its main widget is a combo box that makes available the list
-    of presets detected in the project and content plugins.
+class UAssetSetting(Setting):
+    '''
+    A setting which points to an Unreal asset.
+
+    Its main widget is a combo box that enumerates certain project/plugin asset types.
+    This is an abstract base class; concrete subclasses override `_classnames()` to specify their uasset class names.
 
-    Its "refresh" button triggers an asset traversal to refresh
-    the options.
+    Its "refresh" button triggers an asset traversal to refresh the options.
     '''
 
     def _create_widgets(self, override_device_name=None):
@@ -425,7 +427,10 @@ class LiveLinkPresetSetting(Setting):
         noneasset = Asset(id=0, project=None, assettype=None, gamepath='', name='', localpath='')
         combo.addItem(noneasset.name, noneasset)
 
-        project = SBCache().query_or_create_project(CONFIG.UPROJECT_PATH.get_value())
+        project = SBCache().query_project(CONFIG.UPROJECT_PATH.get_value())
+        if project is None:
+            return
+
         assets = SBCache().query_assets_by_classname(
             project=project,
             classnames=self._classnames())
@@ -444,9 +449,10 @@ class LiveLinkPresetSetting(Setting):
                 combo.setCurrentIndex(item_idx)
                 break
 
-    def _classnames(self) -> list[str]:
+    @abstractmethod
+    def _classnames(self) -> Iterable[str]:
         ''' Return the valid class names of this type of asset '''
-        return DeviceUnreal.LIVELINKPRESET_CLASS_NAMES
+        assert False, "Must be implemented by derived classes"
 
     def _on_setting_changed(
             self, new_value: str,
@@ -474,10 +480,18 @@ class LiveLinkPresetSetting(Setting):
                     break
 
 
-class MediaProfileSetting(LiveLinkPresetSetting):
+class LiveLinkPresetSetting(UAssetSetting):
+    ''' Container of the LiveLinkPreset setting.
+    '''
+    def _classnames(self) -> Iterable[str]:
+        ''' Return the valid class names of this type of asset '''
+        return DeviceUnreal.LIVELINKPRESET_CLASS_NAMES
+
+
+class MediaProfileSetting(UAssetSetting):
     ''' Container of the MediaProfile setting.
     '''
-    def _classnames(self) -> list[str]:
+    def _classnames(self) -> Iterable[str]:
         ''' Return the valid class names of this type of asset '''
         return DeviceUnreal.MEDIAPROFILE_CLASS_NAMES
 
@@ -3249,8 +3263,6 @@ class DeviceUnreal(Device):
         progressDiag.setWindowTitle('Unreal Asset Finder')
         progressDiag.setModal(True)
         progressDiag.setMinimumDuration(1000)  # time before it shows up
-        progressDiag.setCancelButton(None)
-
         # Looks much better without the window frame.
         progressDiag.setWindowFlag(QtCore.Qt.FramelessWindowHint)
 
@@ -3263,6 +3275,10 @@ class DeviceUnreal(Device):
             for path in search_paths:
                 for dirpath, _, file_names in os.walk(path):
                     for file_name in file_names:
+                        if progressDiag.wasCanceled():
+                            done_event.set()
+                            return
+
                         if not file_name.lower().endswith(('.uasset', '.ndisplay')):
                             continue
 
@@ -3288,15 +3304,20 @@ class DeviceUnreal(Device):
         thread.start()
 
         # wait for the event to be set or the progress dialog to be canceled
-        while not done_event.is_set() and not progressDiag.wasCanceled():
+        while not done_event.is_set():
             progressDiag.setValue(progressDiag.value() + 1)
             QtWidgets.QApplication.processEvents()
             time.sleep(0.050)  # The worker thread will run faster if we sleep.
 
+        find_canceled = progressDiag.wasCanceled()
+
         progressDiag.close()
 
         thread.join()
 
+        if find_canceled:
+            return
+
         # process the assets in a multi-threaded fashion
 
         # show a progress bar if it is taking more a trivial amount of time
@@ -3307,7 +3328,6 @@ class DeviceUnreal(Device):
         progressDiag.setModal(True)
         progressDiag.setMinimumDuration(1000)  # time before it shows up
         progressDiag.setRange(0, len(assets))
-        progressDiag.setCancelButton(None)
 
         # Looks much better without the window frame.
         progressDiag.setWindowFlag(QtCore.Qt.FramelessWindowHint)
@@ -3317,7 +3337,7 @@ class DeviceUnreal(Device):
             + DeviceUnreal.LIVELINKPRESET_CLASS_NAMES \
             + DeviceUnreal.MEDIAPROFILE_CLASS_NAMES
 
-        def validateInterestingAsset(asset):                                                                                                
+        def validateInterestingAsset(asset):
             ''' Returns the asset if it is an interesting asset '''
 
             with open(asset['path'], 'rb') as file:
@@ -3336,11 +3356,16 @@ class DeviceUnreal(Device):
 
         with concurrent.futures.ThreadPoolExecutor(
                 max_workers=numThreads) as executor:
+
+            progressDiag.canceled.connect(lambda: executor.shutdown(cancel_futures=True))
+
             futures = [
                 executor.submit(validateInterestingAsset, asset)
                 for asset in assets]
 
             for future in concurrent.futures.as_completed(futures):
+                if progressDiag.wasCanceled():
+                    return
 
                 # Update progress bar.
                 doneAssetCount += 1
@@ -3364,6 +3389,9 @@ class DeviceUnreal(Device):
                 except Exception:
                     pass
 
+        if progressDiag.wasCanceled():
+            return
+
         # close progress bar window
         progressDiag.close()
 
diff --git a/Engine/Plugins/VirtualProduction/Switchboard/Source/Switchboard/switchboard/sbcache.py b/Engine/Plugins/VirtualProduction/Switchboard/Source/Switchboard/switchboard/sbcache.py
--- a/Engine/Plugins/VirtualProduction/Switchboard/Source/Switchboard/switchboard/sbcache.py
+++ b/Engine/Plugins/VirtualProduction/Switchboard/Source/Switchboard/switchboard/sbcache.py
@@ -1,7 +1,8 @@
 # Copyright Epic Games, Inc. All Rights Reserved.
 
-import sqlite3
 import pathlib
+import sqlite3
+from typing import Iterable
 
 
 ROOT_SBCACHE_PATH = pathlib.Path(__file__).parent.with_name('cache')
@@ -230,7 +231,7 @@ class SBCache:
 
             return Project(id=cur.lastrowid, path=filepath)
 
-    def query_project(self, filepath: str) -> Project:
+    def query_project(self, filepath: str) -> Project | None:
         ''' Returns the project with the given filepath. None if non-existent'''
 
         cur = self.con.execute('''
@@ -304,7 +305,7 @@ class SBCache:
 
             return AssetType(id=cur.lastrowid, classname=classname)
 
-    def query_assets_by_classname(self, project: Project, classnames: list[str]) -> list[Asset]:
+    def query_assets_by_classname(self, project: Project, classnames: Iterable[str]) -> list[Asset]:
         ''' Returns a list of assets that are of any of the given classnames. '''
 
         # Convert the classnames to asset type ids
