Switchboard: Allow manually adding a single level, avoiding a potentially long map discovery.

When adding a config from Unreal Editor, only the single map is now added and does not incur in the map discovery cost. Most of the time, one is only interested in running that one map.

Until SB is aware of external project content mounting, this can be used to reference maps from another project under the assumption that they are mounted.

#rnx

diff --git a/Engine/Plugins/VirtualProduction/Switchboard/Source/Switchboard/switchboard/config.py b/Engine/Plugins/VirtualProduction/Switchboard/Source/Switchboard/switchboard/config.py
--- a/Engine/Plugins/VirtualProduction/Switchboard/Source/Switchboard/switchboard/config.py
+++ b/Engine/Plugins/VirtualProduction/Switchboard/Source/Switchboard/switchboard/config.py
@@ -3187,6 +3187,102 @@ class Config(object):
 
         return levels
 
+    def deduce_umap_gamepath_from_filepath(self, umap_filepath: Path) -> Optional[str]:
+        """
+        Convert an absolute .umap file path into its Unreal “game path” by:
+        1. Finding the nearest ancestor folder named "Content"
+        2. Checking the parent of that "Content" folder for either:
+            - A .uproject file  → project content → prefix "/Game"
+            - A .uplugin file   → plugin content  → prefix "/<PluginName>"
+        3. Stripping off the filesystem bits up through "Content/" and any subfolders,
+            then joining the remainder with forward slashes and prepending the prefix.
+
+        Args:
+            umap_path (Path):
+                The full filesystem path to a `.umap` file.
+
+        Returns:
+            Optional[str]:
+                The normalized game-path string (e.g. "/Game/Maps/MyMap" or
+                "/MyPlugin/Levels/MyMap"), or `None` if no Content ancestor or
+                no `.uproject`/`.uplugin` found.
+        """
+        parts = umap_filepath.parts
+
+        # find the last occurrence of "Content"
+        try:
+            idx = len(parts) - 1 - parts[::-1].index("Content")
+        except ValueError:
+            LOGGER.warning(f"No 'Content' folder in path: {umap_filepath}")
+            return None
+
+        # Everything after "Content"
+        rel_parts = parts[idx+1:]
+        if not rel_parts:
+            # e.g. .../Content/ itself, unlikely
+            LOGGER.warning(f"No subpath under 'Content' in: {umap_filepath}")
+            return None
+
+        content_dir = Path(*parts[:idx+1])
+        root_dir    = content_dir.parent
+
+        # Figure out if it is a project or plugin map to establish the right prefix.
+        prefix: Optional[str] = None
+
+        # project-level?
+        if any(root_dir.glob("*.uproject")):
+            prefix = "/Game"
+        else:
+            plugins = list(root_dir.glob("*.uplugin"))
+            if plugins:
+                # take the first plugin descriptor name as plugin name
+                plugin_name = plugins[0].stem
+                prefix = f"/{plugin_name}"
+            else:
+                LOGGER.warning(f"No .uproject or .uplugin under {root_dir}")
+                return None
+
+        # strip off extension and join with forward slashes
+        # e.g. rel_parts = ("Maps", "MyMap.umap") → ["Maps","MyMap"]
+        rel_no_ext = [
+            p if not p.lower().endswith(".umap")
+            else p[:-5]
+            for p in rel_parts
+        ]
+
+        umap_gamepath = prefix + "/" + "/".join(rel_no_ext)
+        return umap_gamepath
+
+    def add_single_level(self, umap_gamepath: Path) -> List[str]:
+        """
+        Add the given to the SBCache.
+
+        Args:
+            umap_gamepath (Path):  
+                Level's game path
+        """
+
+        if not umap_gamepath:
+            return []
+
+        path_str = str(umap_gamepath)
+
+        current = self.get_levels()
+        current = set(current)
+
+        if path_str not in current:
+            LOGGER.info(f'Level added via browse: "{path_str}"')
+            levels = list(current) + [path_str]
+        else:
+            levels = list(current)
+
+        # cache updated list of maps
+        project = SBCache().query_or_create_project(self.UPROJECT_PATH.get_value())
+        cmaps = [Map(id=None, project=None, gamepath=map) for map in levels]
+        SBCache().update_project_maps(project=project, maps=cmaps)
+
+        return levels
+
     def _find_levels(self, levels: list[str]) -> None:
         '''
         Returns a list of full level paths in an Unreal Engine project and
@@ -3215,11 +3311,12 @@ class Config(object):
                 (unreal_content_plugin,
                  unreal_content_plugin.plugin_content_path))
 
-        # Pre-filtering by .umap extension is faster than fnmatch on all files
-        # Currently there are no use cases where a different extension is of interest.
         def get_umap_files(dirpath) -> list:
             ''' Returns a list with all *.map files in the recursively searched dirpath
 
+            Pre-filtering by .umap extension is faster than fnmatch on all files
+            Currently there are no use cases where a different extension is of interest.
+
             Parameters
             ----------
             dirpath: str
diff --git a/Engine/Plugins/VirtualProduction/Switchboard/Source/Switchboard/switchboard/switchboard_dialog.py b/Engine/Plugins/VirtualProduction/Switchboard/Source/Switchboard/switchboard/switchboard_dialog.py
--- a/Engine/Plugins/VirtualProduction/Switchboard/Source/Switchboard/switchboard/switchboard_dialog.py
+++ b/Engine/Plugins/VirtualProduction/Switchboard/Source/Switchboard/switchboard/switchboard_dialog.py
@@ -353,7 +353,16 @@ class SwitchboardDialog(QtCore.QObject):
         self.refresh_levels_button.setProperty("frameless", True)
         self.refresh_levels_button.setToolTip("Refresh level list")
 
+        self.browse_to_level_button = sb_widgets.ControlQPushButton()
+        self.browse_to_level_button.setMaximumSize(18, 22)
+        self.browse_to_level_button.setText('...')
+        font = self.browse_to_level_button.font()
+        font.setBold(True)
+        self.browse_to_level_button.setFont(font)
+        self.browse_to_level_button.setToolTip("Browse to level")
+
         self.window.horizontalLayout_7.addWidget(self.level_combo_box)
+        self.window.horizontalLayout_7.addWidget(self.browse_to_level_button)
         self.window.horizontalLayout_7.addWidget(self.refresh_levels_button)
 
         self.shoot = 'Default'
@@ -417,6 +426,7 @@ class SwitchboardDialog(QtCore.QObject):
         self.window.sequence_line_edit.textChanged.connect(self._set_sequence)
         self.level_combo_box.currentIndexChanged.connect(self._on_selected_level_changed)
         self.refresh_levels_button.clicked.connect(self.refresh_levels_incremental)
+        self.browse_to_level_button.clicked.connect(self.browse_to_level_and_add_it)
         self.window.project_cl_combo_box.currentTextChanged.connect(self._set_project_changelist)
         self.window.engine_cl_combo_box.currentIndexChanged.connect(
             lambda _: self._set_engine_changelist(self.window.engine_cl_combo_box.currentText()))
@@ -2286,6 +2296,53 @@ class SwitchboardDialog(QtCore.QObject):
         level_combo = self.level_combo_box
         return [self._get_level_from_combo_box(i) for i in range(1, level_combo.count())] # skip DEFAULT_MAP_TEXT
 
+    def browse_to_level_and_add_it(self) -> Optional[Path]:
+        """
+        Show a file‐open dialog for .umap files, then add & select the chosen level.
+
+        Returns:
+            The gamepath to the chosen .umap if one was picked, otherwise None.
+        """
+
+        maps_path = os.path.normpath(
+            os.path.join(
+                CONFIG.get_project_content_dir(),
+                CONFIG.MAPS_PATH.get_value())
+        )
+        
+        file_path, _ = QtWidgets.QFileDialog.getOpenFileName(
+            parent=None,
+            caption="Browse to Level",
+            dir=maps_path,
+            filter="Unreal Level Files (*.umap)",
+            selectedFilter="",
+            options=QtWidgets.QFileDialog.Options()
+        )
+
+        if not file_path:
+            return None
+
+        umap_filepath = Path(file_path)
+        
+        # Check extension
+        if umap_filepath.suffix.lower() != ".umap":
+            LOGGER.warning(f"Selected file isn’t a .umap: {umap_filepath}")
+            return None
+
+        umap_gamepath = CONFIG.deduce_umap_gamepath_from_filepath(umap_filepath)
+
+        if not umap_gamepath:
+            return None
+
+        levels = CONFIG.add_single_level(umap_gamepath)
+
+        # Update the UI
+
+        self.refresh_levels(levels)
+        self.level = umap_gamepath  # saves the CONFIG if the level changed.
+
+        return umap_gamepath
+
     def refresh_levels_incremental(self):
         ''' Wrapper around `refresh_levels` with the following differences:
             - Only resaves config if the selected level was removed (as opposed
