Switchboard: Allow manually adding a single level, avoiding a potentially long map discovery.

When adding a config from Unreal Editor, only the single map is now added and does not incur in the map discovery cost. Most of the time, one is only interested in running that one map.

Until SB is aware of external project content mounting, this can be used to reference maps from another project under the assumption that they are mounted.

#rnx

diff --git a/Engine/Plugins/VirtualProduction/Switchboard/Source/Switchboard/switchboard/config.py b/Engine/Plugins/VirtualProduction/Switchboard/Source/Switchboard/switchboard/config.py
--- a/Engine/Plugins/VirtualProduction/Switchboard/Source/Switchboard/switchboard/config.py
+++ b/Engine/Plugins/VirtualProduction/Switchboard/Source/Switchboard/switchboard/config.py
@@ -3188,38 +3188,34 @@ class Config(object):
         return levels
 
     def deduce_umap_gamepath_from_filepath(self, umap_filepath: Path) -> Optional[str]:
-        """
-        Convert an absolute .umap file path into its Unreal “game path” by:
-        1. Finding the nearest ancestor folder named "Content"
-        2. Checking the parent of that "Content" folder for either:
-            - A .uproject file  → project content → prefix "/Game"
-            - A .uplugin file   → plugin content  → prefix "/<PluginName>"
-        3. Stripping off the filesystem bits up through "Content/" and any subfolders,
-            then joining the remainder with forward slashes and prepending the prefix.
+        '''
+        Converts a .umap filepath a gamepath. It does this by finding its nearest "Content"
+        ancestor and determining it it belongs to a project or to a uplugin, because that
+        determines the root of the gamepath.
 
         Args:
-            umap_path (Path):
-                The full filesystem path to a `.umap` file.
+            umap_filepath (Path):
+                The filepath to the level.
 
         Returns:
             Optional[str]:
-                The normalized game-path string (e.g. "/Game/Maps/MyMap" or
-                "/MyPlugin/Levels/MyMap"), or `None` if no Content ancestor or
-                no `.uproject`/`.uplugin` found.
-        """
+                The normalized gamepath string (e.g. "/Game/.../MyMap" or "/SomePlugin/.../MyMap"), 
+                or None if something didn't work out.
+        '''
+
         parts = umap_filepath.parts
 
-        # find the last occurrence of "Content"
+        # Find the last occurrence of "Content" in the path
         try:
             idx = len(parts) - 1 - parts[::-1].index("Content")
         except ValueError:
             LOGGER.warning(f"No 'Content' folder in path: {umap_filepath}")
             return None
 
-        # Everything after "Content"
+        # Use the relative paths under "Content" as those are in the gamepath
         rel_parts = parts[idx+1:]
+
         if not rel_parts:
-            # e.g. .../Content/ itself, unlikely
             LOGGER.warning(f"No subpath under 'Content' in: {umap_filepath}")
             return None
 
@@ -3227,30 +3223,34 @@ class Config(object):
         root_dir    = content_dir.parent
 
         # Figure out if it is a project or plugin map to establish the right prefix.
+
         prefix: Optional[str] = None
 
-        # project-level?
         if any(root_dir.glob("*.uproject")):
             prefix = "/Game"
         else:
             plugins = list(root_dir.glob("*.uplugin"))
+
             if plugins:
-                # take the first plugin descriptor name as plugin name
+
+                # Use the first found uplugin as the plugin name used in the gamepath
                 plugin_name = plugins[0].stem
                 prefix = f"/{plugin_name}"
+
             else:
                 LOGGER.warning(f"No .uproject or .uplugin under {root_dir}")
                 return None
 
-        # strip off extension and join with forward slashes
-        # e.g. rel_parts = ("Maps", "MyMap.umap") → ["Maps","MyMap"]
-        rel_no_ext = [
-            p if not p.lower().endswith(".umap")
-            else p[:-5]
-            for p in rel_parts
+        # Strip off extension and join with forward slashes
+
+        relativepath_without_ext = [
+            part if not part.lower().endswith(".umap")
+            else part[:-5]
+            for part in rel_parts
         ]
 
-        umap_gamepath = prefix + "/" + "/".join(rel_no_ext)
+        umap_gamepath = prefix + "/" + "/".join(relativepath_without_ext)
+
         return umap_gamepath
 
     def add_single_level(self, umap_gamepath: Path) -> List[str]:
