Switchboard: Skip looking for .uplugin files under Content/ folders.

This avoids searching in the Project's Content/ folder which could considerably slow down uplugin discovery in large projects.

Also downgraded from error to info the case of searching for maps in uplugin folders without content.

#rnx

diff --git a/Engine/Plugins/VirtualProduction/Switchboard/Source/Switchboard/switchboard/config.py b/Engine/Plugins/VirtualProduction/Switchboard/Source/Switchboard/switchboard/config.py
--- a/Engine/Plugins/VirtualProduction/Switchboard/Source/Switchboard/switchboard/config.py
+++ b/Engine/Plugins/VirtualProduction/Switchboard/Source/Switchboard/switchboard/config.py
@@ -3342,7 +3342,7 @@ class Config(object):
             try:
                 umaps = get_umap_files(maps_path)
             except FileNotFoundError:
-                LOGGER.error(f"Could not get levels from {maps_path}")
+                LOGGER.info(f"Could not get levels from {maps_path}")
                 continue
 
             for umap in umaps:
diff --git a/Engine/Plugins/VirtualProduction/Switchboard/Source/Switchboard/switchboard/ue_plugin_utils.py b/Engine/Plugins/VirtualProduction/Switchboard/Source/Switchboard/switchboard/ue_plugin_utils.py
--- a/Engine/Plugins/VirtualProduction/Switchboard/Source/Switchboard/switchboard/ue_plugin_utils.py
+++ b/Engine/Plugins/VirtualProduction/Switchboard/Source/Switchboard/switchboard/ue_plugin_utils.py
@@ -233,14 +233,23 @@ class UnrealPluginManager:
 
         start_time = time.perf_counter()
 
+        # We don't expect to find plugins in the following folders, so we can prune those subtrees.
+        skip_dirs = {'content'} 
+
         found_plugins: list[Path] = []
+        suffix = '.uplugin'
 
-        for dirpath, dirs, files in os.walk(in_path):
-            for file in files:
-                if file.casefold().endswith('.uplugin'):
-                    found_plugins.append(Path(dirpath) / file)
-                    dirs.clear()  # Plugins don't nest; don't descend this tree
-                    continue
+        for dirpath, dirs, files in os.walk(in_path, topdown=True):
+            # prune subtrees whose name is in skip_dirs
+            dirs[:] = [d for d in dirs if d.lower() not in skip_dirs]
+
+            for fname in files:
+                if fname.lower().endswith(suffix):
+                    found_plugins.append(Path(dirpath) / fname)
+
+                     # Plugins don't nest; don't descend this tree
+                    dirs.clear()
+                    break
 
         end_time = time.perf_counter()
 
