Switchboard: Level finder can be stopped now.

If stopped, the partial search results get merged with the existing cache of levels for the given project. This allows the user to stop the scan if it normally takes a long time, but without losing the results of a previous scan.

If not stopped, it will replace the data in the cache as usual.

#rnx

diff --git a/Engine/Plugins/VirtualProduction/Switchboard/Source/Switchboard/switchboard/config.py b/Engine/Plugins/VirtualProduction/Switchboard/Source/Switchboard/switchboard/config.py
--- a/Engine/Plugins/VirtualProduction/Switchboard/Source/Switchboard/switchboard/config.py
+++ b/Engine/Plugins/VirtualProduction/Switchboard/Source/Switchboard/switchboard/config.py
@@ -3153,12 +3153,20 @@ class Config(object):
 
         # show a progress bar if it is taking more a trivial amount of time
         progressDiag = QtWidgets.QProgressDialog(
-            'Finding levels...', 'Cancel', 0, 0, parent=None)
-
+            labelText='Finding levels...',
+            cancelButtonText='Stop',
+            minimum=0, 
+            maximum=0, 
+            parent=None
+        )
+        
         progressDiag.setWindowTitle('Unreal Level Finder')
         progressDiag.setModal(True)
         progressDiag.setMinimumDuration(1000)  # time before it shows up
-        progressDiag.setCancelButton(None)
+
+        # Allow cancellation but still letting partial results through.
+        stop_event = threading.Event()
+        progressDiag.canceled.connect(stop_event.set)
 
         # Looks much better without the window frame.
         progressDiag.setWindowFlag(QtCore.Qt.FramelessWindowHint)
@@ -3169,7 +3177,7 @@ class Config(object):
 
         # our convenience worker function
         def find_level_work():
-            self._find_levels(levels)
+            self._find_levels(levels, stop_event)
             done_event.set()
 
         thread = threading.Thread(target=find_level_work)
@@ -3283,7 +3291,7 @@ class Config(object):
 
         return levels
 
-    def _find_levels(self, levels: list[str]) -> None:
+    def _find_levels(self, levels: list[str], stop_event: Optional[threading.Event] = None) -> None:
         '''
         Returns a list of full level paths in an Unreal Engine project and
         in plugins such as:
@@ -3328,21 +3336,31 @@ class Config(object):
                     List of found map entries.
             '''
             umap_files = []
+
             for entry in os.scandir(dirpath):
+
+                if stop_event and stop_event.is_set():
+                    break
+
                 if entry.is_file() and entry.name.endswith(".umap"):
                     umap_files.append(entry)
                 elif entry.is_dir():
                     umap_files.extend(get_umap_files(entry.path))  # recursive
+
             return umap_files
 
         maps_filter = self.MAPS_FILTER.get_value()
 
         for (unreal_content_plugin, maps_path) in search_paths:
 
+            # Stop the search if the user doesn't want to wait any longer.
+            if stop_event and stop_event.is_set():
+                break
+
             try:
                 umaps = get_umap_files(maps_path)
             except FileNotFoundError:
-                LOGGER.info(f"Could not get levels from {maps_path}")
+                # This is normal for non-content plugins.
                 continue
 
             for umap in umaps:
@@ -3356,12 +3374,30 @@ class Config(object):
                 else:
                     LOGGER.warning(f"No game path for {umap.path}")
 
-        levels.sort()
+        # If the search was stopped, merge the search results with cache.
+        # Otherwise replace the cache with the search results.
 
-        # cache updated list of maps
         project = SBCache().query_or_create_project(self.UPROJECT_PATH.get_value())
-        cmaps = [Map(id=None, project=None, gamepath=map) for map in levels]
-        SBCache().update_project_maps(project=project, maps=cmaps)
+
+        if stop_event and stop_event.is_set():
+
+            LOGGER.info("Stopped level scan early, this will merge partial search results with the existing cache.")
+            
+            existing_maps = SBCache().query_maps(project)
+            existing_paths = {map.gamepath for map in existing_maps}
+
+            merged_paths = sorted(existing_paths.union(levels))
+            merged_maps = [Map(id=None, project=None, gamepath=path) for path in merged_paths]
+
+            # update levels since the caller expects the merged results in there.
+            levels.clear()
+            levels.extend(merged_paths)
+            
+            SBCache().update_project_maps(project=project, maps=merged_maps)
+        else:
+            levels.sort()
+            cmaps = [Map(id=None, project=None, gamepath=map) for map in levels]
+            SBCache().update_project_maps(project=project, maps=cmaps)
 
     def get_levels(self) -> list[str]:
         ''' Returns a list with all *.map files from the cache '''
