Swichboard: Allow copy-paste of asset references from UE's content browser.

Currently only umaps are supported, but the intention is to support other interesting assets. The benefit is that if a new asset gets created, it can be used by switchboard right away without full asset or level discovery, which can be lengthy if it is a large project.

#rnx

diff --git a/Engine/Plugins/VirtualProduction/Switchboard/Source/Switchboard/switchboard/switchboard_dialog.py b/Engine/Plugins/VirtualProduction/Switchboard/Source/Switchboard/switchboard/switchboard_dialog.py
--- a/Engine/Plugins/VirtualProduction/Switchboard/Source/Switchboard/switchboard/switchboard_dialog.py
+++ b/Engine/Plugins/VirtualProduction/Switchboard/Source/Switchboard/switchboard/switchboard_dialog.py
@@ -282,6 +282,12 @@ class SwitchboardDialog(QtCore.QObject):
         # Create an instance of the UI class
         self.window = Ui_switchboard_window()
 
+        # Accept copy-paste, intended for uassets or umaps from Unreal Engine's Content Browser
+        QtGui.QShortcut(QtGui.QKeySequence.Paste, self.window).activated.connect(self.on_paste)
+
+        # Avoid other widgets getting focus automaticaly so we can catch the paste shortcut
+        self.window.centralwidget.setFocusPolicy(QtCore.Qt.FocusPolicy.ClickFocus)
+
         # Add Tools Menu
         self.add_tools_menu()
 
@@ -532,6 +538,63 @@ class SwitchboardDialog(QtCore.QObject):
             lambda _: self.p4_refresh_project_cl()
         )
 
+    def on_paste(self):
+        ''' Catches pasting of content browser assets.
+        The main purpose is to auto-add them to switchboard's caches of assets and levels.
+
+        Expected syntax:
+            /Script/Module.Class'Gamepath.Name'
+        
+        E.g. umap:
+            /Script/Engine.World'/Game/Tests/AA/LLF/L_LLF.L_LLF
+
+        When multiple objects are selected, they are listed one per line.
+        '''
+
+        clipboard = QtGui.QGuiApplication.clipboard()
+        mime = clipboard.mimeData()
+
+        if mime.hasText():
+
+            text = mime.text()
+            
+            asset_reference_re = re.compile(r"""
+                ^/Script/                              # /Script/
+                (?P<module>[A-Za-z_]\w+)               # Module
+                \.                                     # .
+                (?P<classname>[A-Za-z_]\w+)            # Class
+                '                                      # '
+                (?P<gamepath>[A-Za-z0-9_/]+)           # Gamepath
+                \.                                     # .
+                (?P<name>[A-Za-z0-9_]+)                # Name
+                '                                      # '
+                $""", re.VERBOSE)
+            
+            # When selecting multiple assets, they are listed on separate lines.
+            for line in text.splitlines():
+
+                match = asset_reference_re.match(line)
+
+                if not match:
+                    LOGGER.warning(f"Invalid asset reference: {line!r}")
+                    continue
+                
+                module    = match.group("module")
+                classname = match.group("classname")
+                gamepath  = match.group("gamepath")
+
+                if f"{module}.{classname}" == "Engine.World":
+
+                    levels = CONFIG.add_single_level(gamepath)
+
+                    # Update the UI
+                    self.refresh_levels(levels)
+                    self.level = gamepath  # saves the CONFIG if the level changed.
+
+                    LOGGER.info(f"Added level '{gamepath}'")
+
+                # Note: One could also handle other interesting assets, such as media profiles, etc.
+
     def _try_change_address(self):
         new_value = self.window.current_address_value.text()
         old_address = SETTINGS.ADDRESS.get_value()
