Switchboard: Allow copy-pasting of uassets from the content browser of UE onto Switchboard.

Currently supports:
* Media Profiles
* Live Link Presets
* nDisplay configs

This allows you to quickly add assets of interest without needing to re-scan all of the UE project content, which can take some time depending on its size.

#rnx

diff --git a/Engine/Plugins/VirtualProduction/Switchboard/Source/Switchboard/switchboard/devices/unreal/plugin_unreal.py b/Engine/Plugins/VirtualProduction/Switchboard/Source/Switchboard/switchboard/devices/unreal/plugin_unreal.py
--- a/Engine/Plugins/VirtualProduction/Switchboard/Source/Switchboard/switchboard/devices/unreal/plugin_unreal.py
+++ b/Engine/Plugins/VirtualProduction/Switchboard/Source/Switchboard/switchboard/devices/unreal/plugin_unreal.py
@@ -811,6 +811,7 @@ class DeviceUnreal(Device):
     # A stand-in QObject to emit signals from classmethods.
     class StaticSignals(QtCore.QObject):
         ugs_config_updated_signal = QtCore.Signal(ugs_utils.IniParser)
+        clipboard_level_added_signal = QtCore.Signal(str)
 
     static_signals = StaticSignals()
 
@@ -3428,6 +3429,7 @@ class DeviceUnreal(Device):
                 continue  # We will have logged an error already.
 
             assettype = SBCache().query_or_create_assettype(asset_classnames[idx])
+
             asset = Asset(
                 id=0,
                 project=project,
@@ -3441,6 +3443,136 @@ class DeviceUnreal(Device):
         # save to the cache
         SBCache().update_project_assets(project=project, assets=assets)
 
+    @classmethod
+    def handle_clipboard_mime_text(cls, text:str):
+        ''' Catches pasting of content browser assets.
+        The main purpose is to auto-add them to switchboard's caches of assets and levels.
+
+        Expected syntax:
+            /Script/Module.Class'Gamepath.Name'
+        
+        E.g. umap:
+            /Script/Engine.World'/Game/Tests/AA/LLF/L_LLF.L_LLF
+
+        When multiple objects are selected, they are listed one per line.
+        '''  
+        asset_reference_re = re.compile(r"""
+            ^/Script/                              # /Script/
+            (?P<module>[A-Za-z_]\w+)               # Module
+            \.                                     # .
+            (?P<classname>[A-Za-z_]\w+)            # Class
+            '                                      # '
+            (?P<gamepath>[A-Za-z0-9_/]+)           # Gamepath
+            \.                                     # .
+            (?P<name>[A-Za-z0-9_]+)                # Name
+            '                                      # '
+            $""", re.VERBOSE)
+
+        INTERESTING_CLASS_NAMES = \
+            DeviceUnreal.NDISPLAY_CLASS_NAMES \
+            + DeviceUnreal.LIVELINKPRESET_CLASS_NAMES \
+            + DeviceUnreal.MEDIAPROFILE_CLASS_NAMES
+        
+        assets = []
+        levels = []
+
+        project = SBCache().query_or_create_project(CONFIG.UPROJECT_PATH.get_value())
+
+        # When selecting multiple assets, they are listed on separate lines.
+        for line in text.splitlines():
+
+            match = asset_reference_re.match(line)
+
+            if not match:
+                LOGGER.warning(f"Invalid asset reference: {line!r}")
+                continue
+            
+            module    = match.group("module")
+            classname = match.group("classname")
+            gamepath  = match.group("gamepath")
+            name      = match.group("name")
+
+            namespaced_classname = f"/Script/{module}.{classname}"
+
+            if namespaced_classname == "/Script/Engine.World":
+                levels.append(gamepath)
+
+            elif namespaced_classname in INTERESTING_CLASS_NAMES:
+                filepath = CONFIG.ue_plugin_mgr.content_to_file_path(gamepath)
+
+                # The class type implies the file extension.
+                uasset_ext = '.uasset'
+
+                file_exists = False
+
+                if filepath:
+                    # Append the extension it since content_to_file_path(..) doesn't.
+                    filepath = str(filepath) + uasset_ext
+                    file_exists = os.path.exists(filepath)
+                else:
+                    filepath = ''
+
+                if not file_exists:
+                    # Don't save invalid filepaths
+                    filepath = ''
+
+                    # The only case that requires a valid filepath is nDisplay configs, since they are parsed
+                    # during launch. The other asset references are ok with just a valid gamepath.
+                    if namespaced_classname in DeviceUnreal.NDISPLAY_CLASS_NAMES:
+                        LOGGER.warning(f"Skipping nDisplay config '{gamepath}' because its file on disk could not be located")
+                        continue
+                    else:
+                        LOGGER.warning(f"Could not resolve to an existing filepath for '{gamepath}'")
+
+                assettype = SBCache().query_or_create_assettype(namespaced_classname)
+                asset_name = name + uasset_ext
+
+                asset = Asset(
+                    id=0,
+                    project=project,
+                    assettype=assettype,
+                    gamepath=str(gamepath),
+                    name=asset_name,
+                    localpath=filepath
+                )
+                
+                assets.append(asset)
+
+        # Add levels
+        if len(levels):
+
+            levels_text = '\n'.join(levels)
+
+            reply = QtWidgets.QMessageBox.question(
+                None,
+                "Add Unreal Maps",
+                f"The following maps will be added:\n\n{levels_text}",
+                QtWidgets.QMessageBox.Ok | QtWidgets.QMessageBox.Cancel
+            )
+
+            if reply == QtWidgets.QMessageBox.Ok:
+                for level in levels:
+                    CONFIG.add_single_level(level)
+                    cls.static_signals.clipboard_level_added_signal.emit(level)
+                    LOGGER.info(f"Added level '{level}'")
+
+        # Add assets
+        if len(assets):
+
+            assets_text = '\n'.join([asset.name for asset in assets])
+
+            reply = QtWidgets.QMessageBox.question(
+                None,
+                "Add Unreal Assets",
+                f"The following assets will be added:\n\n{assets_text}",
+                QtWidgets.QMessageBox.Ok | QtWidgets.QMessageBox.Cancel
+            )
+
+            SBCache().merge_project_assets(project, assets)
+
+            for asset in assets:
+                LOGGER.info(f"Added '{asset.name}'")
+
 
 def parse_unreal_tag_file(file_content):
     tags = []
diff --git a/Engine/Plugins/VirtualProduction/Switchboard/Source/Switchboard/switchboard/sbcache.py b/Engine/Plugins/VirtualProduction/Switchboard/Source/Switchboard/switchboard/sbcache.py
--- a/Engine/Plugins/VirtualProduction/Switchboard/Source/Switchboard/switchboard/sbcache.py
+++ b/Engine/Plugins/VirtualProduction/Switchboard/Source/Switchboard/switchboard/sbcache.py
@@ -283,6 +283,9 @@ class SBCache:
         If it doesn't exist, then it creates it.
         '''
 
+        # Normalize to class names without namespace to avoid duplicates.
+        classname = self.normalize_classname(classname)
+
         # Check if the asset type exists in the database
         cur = self.con.execute('''
             SELECT AssetTypeId, ClassName FROM AssetTypes
@@ -305,6 +308,14 @@ class SBCache:
 
             return AssetType(id=cur.lastrowid, classname=classname)
 
+    def normalize_classname(self, classname: str):
+        ''' Sometimes classnames are fully namespaced, e.g. /Script/Module.Classname and sometimes they aren't.
+        To avoid duplicate assets due to differing classnames, we opt to drop the namespace. Currently all interesting
+        classes in Switchboard don't have namespace collisions.
+        '''
+
+        return classname.split('.')[-1]
+
     def query_assets_by_classname(self, project: Project, classnames: Iterable[str]) -> list[Asset]:
         ''' Returns a list of assets that are of any of the given classnames. '''
 
@@ -358,3 +369,39 @@ class SBCache:
         ''', [(project.id, asset.assettype.id, asset.name, asset.gamepath, asset.localpath) for asset in assets])
 
         self.con.commit()
+
+    def merge_project_assets(self, project: Project, assets: list[Asset]) -> None:
+        ''' Inserts new assets into the Assets table.
+        To resolve conflicts with existing ones (determined by the GamePath alone for the given project), 
+        it updates the other fields. Pre-existing assets are left as they were.
+        '''
+
+        # ensure the unique index we need exists
+        self.con.execute('''
+            CREATE UNIQUE INDEX IF NOT EXISTS idx_assets_project_gamepath
+            ON Assets(ProjectId, GamePath)
+        ''')
+
+        # bulk upsert
+        self.con.executemany('''
+            INSERT INTO Assets
+            (ProjectId, AssetTypeId, Name, GamePath, LocalPath)
+            VALUES (?, ?, ?, ?, ?)
+            ON CONFLICT(ProjectId, GamePath) DO UPDATE
+            SET AssetTypeId = excluded.AssetTypeId,
+                Name        = excluded.Name,
+                LocalPath   = excluded.LocalPath
+        ''',
+            [
+                (
+                    project.id,
+                    asset.assettype.id,
+                    asset.name,
+                    asset.gamepath,
+                    asset.localpath
+                )
+                for asset in assets
+            ]
+        )
+
+        self.con.commit()
diff --git a/Engine/Plugins/VirtualProduction/Switchboard/Source/Switchboard/switchboard/switchboard_dialog.py b/Engine/Plugins/VirtualProduction/Switchboard/Source/Switchboard/switchboard/switchboard_dialog.py
--- a/Engine/Plugins/VirtualProduction/Switchboard/Source/Switchboard/switchboard/switchboard_dialog.py
+++ b/Engine/Plugins/VirtualProduction/Switchboard/Source/Switchboard/switchboard/switchboard_dialog.py
@@ -288,6 +288,9 @@ class SwitchboardDialog(QtCore.QObject):
         # Avoid other widgets getting focus automaticaly so we can catch the paste shortcut
         self.window.centralwidget.setFocusPolicy(QtCore.Qt.FocusPolicy.ClickFocus)
 
+        # Connect to copy-paste events from handling devices
+        DeviceUnreal.static_signals.clipboard_level_added_signal.connect(self.on_clipboard_level_added)
+
         # Add Tools Menu
         self.add_tools_menu()
 
@@ -541,59 +544,24 @@ class SwitchboardDialog(QtCore.QObject):
     def on_paste(self):
         ''' Catches pasting of content browser assets.
         The main purpose is to auto-add them to switchboard's caches of assets and levels.
-
-        Expected syntax:
-            /Script/Module.Class'Gamepath.Name'
-        
-        E.g. umap:
-            /Script/Engine.World'/Game/Tests/AA/LLF/L_LLF.L_LLF
-
-        When multiple objects are selected, they are listed one per line.
         '''
 
         clipboard = QtGui.QGuiApplication.clipboard()
         mime = clipboard.mimeData()
 
         if mime.hasText():
+            DeviceUnreal.handle_clipboard_mime_text(mime.text())
 
-            text = mime.text()
-            
-            asset_reference_re = re.compile(r"""
-                ^/Script/                              # /Script/
-                (?P<module>[A-Za-z_]\w+)               # Module
-                \.                                     # .
-                (?P<classname>[A-Za-z_]\w+)            # Class
-                '                                      # '
-                (?P<gamepath>[A-Za-z0-9_/]+)           # Gamepath
-                \.                                     # .
-                (?P<name>[A-Za-z0-9_]+)                # Name
-                '                                      # '
-                $""", re.VERBOSE)
-            
-            # When selecting multiple assets, they are listed on separate lines.
-            for line in text.splitlines():
-
-                match = asset_reference_re.match(line)
-
-                if not match:
-                    LOGGER.warning(f"Invalid asset reference: {line!r}")
-                    continue
-                
-                module    = match.group("module")
-                classname = match.group("classname")
-                gamepath  = match.group("gamepath")
-
-                if f"{module}.{classname}" == "Engine.World":
-
-                    levels = CONFIG.add_single_level(gamepath)
-
-                    # Update the UI
-                    self.refresh_levels(levels)
-                    self.level = gamepath  # saves the CONFIG if the level changed.
-
-                    LOGGER.info(f"Added level '{gamepath}'")
-
-                # Note: One could also handle other interesting assets, such as media profiles, etc.
+    @QtCore.Slot(object)
+    def on_clipboard_level_added(self, gamepath):
+        ''' Called when a new level gets added. This will typically happen via a copy-paste operation from
+        UE's content browser. When this happens, this dialog needs to refresh the dropdown, and we are also
+        automatically selecting it in the dropdown, as that is usually what we expect the user to do when they are
+        manually adding a new level.
+        '''
+        levels = CONFIG.get_levels()
+        self.refresh_levels(levels)
+        self.level = gamepath
 
     def _try_change_address(self):
         new_value = self.window.current_address_value.text()
diff --git a/Engine/Plugins/VirtualProduction/Switchboard/Source/Switchboard/switchboard/ue_plugin_utils.py b/Engine/Plugins/VirtualProduction/Switchboard/Source/Switchboard/switchboard/ue_plugin_utils.py
--- a/Engine/Plugins/VirtualProduction/Switchboard/Source/Switchboard/switchboard/ue_plugin_utils.py
+++ b/Engine/Plugins/VirtualProduction/Switchboard/Source/Switchboard/switchboard/ue_plugin_utils.py
@@ -209,11 +209,21 @@ class UnrealPluginManager:
 
         if path.is_relative_to(project_content_path):
             rest = path.relative_to(project_content_path)
+
+            # Strip file extension if present
+            if rest.suffix:
+                rest = rest.with_suffix('')
+                
             return UnrealPath(PROJECT_CONTENT_MOUNT_POINT) / rest
 
         for parent in path.parents:
             if plugin := self.get_plugin_by_content_dir(parent):
                 rest = path.relative_to(parent)
+
+                # Strip file extension if present
+                if rest.suffix:
+                    rest = rest.with_suffix('')
+
                 return plugin.mounted_path / rest
 
         logging.warning('Failed to resolve content path "%s"', path)
