Switchboard: Multi-user session name is now saved with the SB project config instead of being a user setting.

MU session names mostly apply to a given config as opposed to being shared between them, and with this change the user will be able to retrieve the last MU session used in a given config. Without this change, users may end up trying to open an invalid session in the newly opened project.

In addition, the default MU session name now includes the date in YYMMDD format, which helps avoid duplicate session names when opening the same project in the future.

As a convenience, when the MU session field is emptied and the session increment "+" button is pressed, a default session name is generated with the pattern MU_{PROJECT}_{YYMMDD}_1. e.g. MU_QAVPTest_250719_1.

Also fixed a warning related to an escape code in a multi-line comment, and removed the project name generation when creating configs from Unreal Editor because project names are generated automatically when creating the config, based on the config name. It is not meant to reflect the name of the .uproject necessarily.

#rnx

diff --git a/Engine/Plugins/VirtualProduction/Switchboard/Source/Switchboard/switchboard/config.py b/Engine/Plugins/VirtualProduction/Switchboard/Source/Switchboard/switchboard/config.py
--- a/Engine/Plugins/VirtualProduction/Switchboard/Source/Switchboard/switchboard/config.py
+++ b/Engine/Plugins/VirtualProduction/Switchboard/Source/Switchboard/switchboard/config.py
@@ -1,6 +1,7 @@
 # Copyright Epic Games, Inc. All Rights Reserved.
 
 import collections
+import datetime
 import fnmatch
 import io
 import json
@@ -2781,6 +2782,12 @@ class Config(object):
         ''' Returns default server name based on current settings '''
         return f'{self.PROJECT_NAME.get_value()}_MU_Server'
 
+    def default_mu_session_name(self):
+        ''' Returns default session name based on current settings '''
+        project_name = self.PROJECT_NAME.get_value()
+        date_str = datetime.datetime.now().strftime('%y%m%d')
+        return f'MU_{project_name}_{date_str}_1'
+
     def init_muserver(self, data={}):
         self.mu_settings = {
             "muserver_server_name": StringSetting(
@@ -2858,6 +2865,12 @@ class Config(object):
                 "muserver_auto_join",
                 "Unreal Multi-user Server Auto-join",
                 data.get('muserver_auto_join', True)
+            ),
+            "muserver_session_name": StringSetting(
+                "muserver_session_name",
+                "Session Name",
+                data.get('muserver_session_name', self.default_mu_session_name()),
+                tool_tip="The name of the multi-user session"
             )
         }
 
@@ -2875,6 +2888,7 @@ class Config(object):
         self.MUSERVER_AUTO_JOIN = self.mu_settings["muserver_auto_join"]
         self.MUSERVER_WORKING_DIR = self.mu_settings["muserver_working_dir"]
         self.MUSERVER_ARCHIVE_DIR = self.mu_settings["muserver_archive_dir"]
+        self.MUSERVER_SESSION_NAME = self.mu_settings["muserver_session_name"]
 
     def save_unreal_insights(self, data):
         data['tracing_enabled'] = self.INSIGHTS_TRACE_ENABLE.get_value()
@@ -2896,6 +2910,7 @@ class Config(object):
         data["muserver_auto_join"] = self.MUSERVER_AUTO_JOIN.get_value()
         data["muserver_archive_dir"] = self.MUSERVER_ARCHIVE_DIR.get_value()
         data["muserver_working_dir"] = self.MUSERVER_WORKING_DIR.get_value()
+        data["muserver_session_name"] = self.MUSERVER_SESSION_NAME.get_value()
 
     def load_plugin_settings(self, device_type, settings):
         ''' Updates plugin settings values with those read from the config file.
@@ -3543,8 +3558,6 @@ class UserSettings(object):
             )
 
         # UI Settings
-        self.MUSERVER_SESSION_NAME = data.get(
-            'muserver_session_name', 'MU_Session').replace(' ','_')
         self.CURRENT_SEQUENCE = data.get('current_sequence', 'Default')
         self.CURRENT_SLATE = data.get('current_slate', 'Scene')
         self.CURRENT_TAKE = data.get('current_take', 1)
@@ -3564,7 +3577,6 @@ class UserSettings(object):
             'config': '',
             'address': self.ADDRESS.get_value(),
             'transport_path': self.TRANSPORT_PATH.get_value(),
-            'muserver_session_name': self.MUSERVER_SESSION_NAME,
             'current_sequence': self.CURRENT_SEQUENCE,
             'current_slate': self.CURRENT_SLATE,
             'current_take': self.CURRENT_TAKE,
diff --git a/Engine/Plugins/VirtualProduction/Switchboard/Source/Switchboard/switchboard/devices/ndisplay/plugin_ndisplay.py b/Engine/Plugins/VirtualProduction/Switchboard/Source/Switchboard/switchboard/devices/ndisplay/plugin_ndisplay.py
--- a/Engine/Plugins/VirtualProduction/Switchboard/Source/Switchboard/switchboard/devices/ndisplay/plugin_ndisplay.py
+++ b/Engine/Plugins/VirtualProduction/Switchboard/Source/Switchboard/switchboard/devices/ndisplay/plugin_ndisplay.py
@@ -1255,7 +1255,7 @@ class DevicenDisplay(DeviceUnreal):
         mu_server = switchboard_application.get_multi_user_server_instance()
         args.extend([
             f'-CONCERTSERVER="{mu_server.configured_server_name()}"',
-            f'-CONCERTSESSION="{SETTINGS.MUSERVER_SESSION_NAME}"',
+            f'-CONCERTSESSION="{CONFIG.MUSERVER_SESSION_NAME.get_value()}"',
             f'-CONCERTDISPLAYNAME="{self.name}"',
             '-CONCERTISHEADLESS',
         ])
diff --git a/Engine/Plugins/VirtualProduction/Switchboard/Source/Switchboard/switchboard/devices/sounddevices/plugin_sounddevices.py b/Engine/Plugins/VirtualProduction/Switchboard/Source/Switchboard/switchboard/devices/sounddevices/plugin_sounddevices.py
--- a/Engine/Plugins/VirtualProduction/Switchboard/Source/Switchboard/switchboard/devices/sounddevices/plugin_sounddevices.py
+++ b/Engine/Plugins/VirtualProduction/Switchboard/Source/Switchboard/switchboard/devices/sounddevices/plugin_sounddevices.py
@@ -152,7 +152,7 @@ class DeviceSoundDevices(Device):
         # Put all drives into file transfer mode
         for i in range(1,5):
             self._send_command(f'setsetting/RecordToDrive{i}=Ethernet%20File%20Transfer')
-        '''
+        r'''
         f'net use * \\\\{self.address} /user:{self.user_name} {self.password}'
         robocopy.exe "W:\wwwroot\MyProject" "\\192.168.0.1\Share\wwwroot\MyProject" *.* /E /XO /XD "App_Data/Search" "*.svn" /XF "sitefinity.log" "Thumbs.db" /NDL /NC /NP
         net use * /delete /yes
diff --git a/Engine/Plugins/VirtualProduction/Switchboard/Source/Switchboard/switchboard/devices/unreal/plugin_unreal.py b/Engine/Plugins/VirtualProduction/Switchboard/Source/Switchboard/switchboard/devices/unreal/plugin_unreal.py
--- a/Engine/Plugins/VirtualProduction/Switchboard/Source/Switchboard/switchboard/devices/unreal/plugin_unreal.py
+++ b/Engine/Plugins/VirtualProduction/Switchboard/Source/Switchboard/switchboard/devices/unreal/plugin_unreal.py
@@ -2207,7 +2207,7 @@ class DeviceUnreal(Device):
                 '-CONCERTAUTOCONNECT ')
         mu_server = switchboard_application.get_multi_user_server_instance()
         command_line_args += (f'-CONCERTSERVER="{mu_server.configured_server_name()}" '
-                              f'-CONCERTSESSION="{SETTINGS.MUSERVER_SESSION_NAME}" '
+                              f'-CONCERTSESSION="{CONFIG.MUSERVER_SESSION_NAME.get_value()}" '
                               f'-CONCERTDISPLAYNAME="{self.name}"')
 
         if CONFIG.INSIGHTS_TRACE_ENABLE.get_value() and not self.exclude_from_insights.get_value():
diff --git a/Engine/Plugins/VirtualProduction/Switchboard/Source/Switchboard/switchboard/switchboard_dialog.py b/Engine/Plugins/VirtualProduction/Switchboard/Source/Switchboard/switchboard/switchboard_dialog.py
--- a/Engine/Plugins/VirtualProduction/Switchboard/Source/Switchboard/switchboard/switchboard_dialog.py
+++ b/Engine/Plugins/VirtualProduction/Switchboard/Source/Switchboard/switchboard/switchboard_dialog.py
@@ -428,7 +428,7 @@ class SwitchboardDialog(QtCore.QObject):
         self.osc_server.dispatcher_map(osc.DATA, self.osc_data)
 
         # Connect UI to methods
-        self.window.multiuser_session_lineEdit.textChanged.connect(self.on_multiuser_session_lineEdit_textChanged)
+        self.window.multiuser_session_lineEdit.editingFinished.connect(self.on_multiuser_session_lineEdit_editingFinished)
         self.window.slate_line_edit.editingFinished.connect(
             lambda: self._set_slate(self.window.slate_line_edit.text()))
         self.window.take_spin_box.valueChanged.connect(self._set_take)
@@ -520,7 +520,7 @@ class SwitchboardDialog(QtCore.QObject):
 
         self.set_config_hooks()
 
-        self.set_multiuser_session_name(f'{SETTINGS.MUSERVER_SESSION_NAME}')
+        self.set_multiuser_session_name(f'{CONFIG.MUSERVER_SESSION_NAME.get_value()}')
 
         # Run the transport queue
         #self.transport_queue_resume()
@@ -914,6 +914,7 @@ class SwitchboardDialog(QtCore.QObject):
         CONFIG.INSIGHTS_TRACE_ARGS.signal_setting_changed.connect(lambda: self.refresh_trace_settings())
         CONFIG.INSIGHTS_STAT_EVENTS.signal_setting_changed.connect(lambda: self.refresh_trace_settings())
         CONFIG.MUSERVER_AUTO_JOIN.signal_setting_changed.connect(lambda: self.refresh_muserver_autojoin())
+        CONFIG.MUSERVER_SESSION_NAME.signal_setting_changed.connect(lambda: self.refresh_muserver_session_name())
         CONFIG.PROJECT_NAME.signal_setting_changed.connect(lambda: self.refresh_window_title())
 
     def refresh_trace_settings(self):
@@ -923,6 +924,9 @@ class SwitchboardDialog(QtCore.QObject):
     def refresh_muserver_autojoin(self):
         self.window.use_device_autojoin_setting_checkbox.setChecked(CONFIG.MUSERVER_AUTO_JOIN.get_value())
 
+    def refresh_muserver_session_name(self):
+        self.set_multiuser_session_name(CONFIG.MUSERVER_SESSION_NAME.get_value())
+
     def on_tracing_settings_changed(self):
         settings = self.window.additional_settings
         trace_tuple = settings.trace_settings()
@@ -1097,6 +1101,7 @@ class SwitchboardDialog(QtCore.QObject):
         self.refresh_levels(CONFIG.get_levels())
         self.update_current_config_text()
         self.refresh_muserver_autojoin()
+        self.refresh_muserver_session_name()
         self.refresh_trace_settings()
         self.refresh_window_title()
 
@@ -1172,6 +1177,7 @@ class SwitchboardDialog(QtCore.QObject):
         self.refresh_levels(CONFIG.get_levels())
         self.update_current_config_text()
         self.refresh_muserver_autojoin()
+        self.refresh_muserver_session_name()
         self.refresh_trace_settings()
         self.refresh_window_title()
 
@@ -1547,14 +1553,23 @@ class SwitchboardDialog(QtCore.QObject):
     def multiuser_session_name(self):
         return self._multiuser_session_name
 
-    def on_multiuser_session_lineEdit_textChanged(self, text):
+    def on_multiuser_session_lineEdit_editingFinished(self):
+        text = self.window.multiuser_session_lineEdit.text()
         self.set_multiuser_session_name(text)
 
     def on_multiuser_session_inc(self):
         '''
-        Increment the session name by 1.
+        Increment the session name by 1, or generate default if empty.
         '''
         current_name = self.multiuser_session_name()
+        
+        # If the session name is empty or just whitespace, generate a default one
+        if not current_name or current_name.strip() == '':
+            new_session_name = CONFIG.default_mu_session_name()
+            self.set_multiuser_session_name(new_session_name)
+            return
+        
+        # Otherwise, use the existing increment logic
         match = re.search(r'\d+$', current_name)
         basename = current_name
         num_as_int = 1
@@ -1581,9 +1596,8 @@ class SwitchboardDialog(QtCore.QObject):
         if self.window.multiuser_session_lineEdit.text() != value:
             self.window.multiuser_session_lineEdit.setText(value)
 
-        if value != SETTINGS.MUSERVER_SESSION_NAME:
-            SETTINGS.MUSERVER_SESSION_NAME = value
-            SETTINGS.save()
+        if value != CONFIG.MUSERVER_SESSION_NAME.get_value():
+            CONFIG.MUSERVER_SESSION_NAME.update_value(value)
 
     @property
     def shoot(self):
