Switchboard: nDisplay launches now default to -nomcp.

Added checkbox in case that is needed in the future.

#rnx

diff --git a/Engine/Plugins/VirtualProduction/Switchboard/Source/Switchboard/switchboard/devices/ndisplay/plugin_ndisplay.py b/Engine/Plugins/VirtualProduction/Switchboard/Source/Switchboard/switchboard/devices/ndisplay/plugin_ndisplay.py
--- a/Engine/Plugins/VirtualProduction/Switchboard/Source/Switchboard/switchboard/devices/ndisplay/plugin_ndisplay.py
+++ b/Engine/Plugins/VirtualProduction/Switchboard/Source/Switchboard/switchboard/devices/ndisplay/plugin_ndisplay.py
@@ -440,6 +440,12 @@ class DevicenDisplay(DeviceUnreal):
             value=False,
             category="UE Settings",
         ),
+        'mcp': BoolSetting(
+            attr_name="mcp",
+            nice_name="MCP Subsystem",
+            value=False,
+            category="UE Settings",
+        ),
         'loading_screen': BoolSetting(
             attr_name="loading_screen",
             nice_name="Loading Screen",
@@ -1039,6 +1045,12 @@ class DevicenDisplay(DeviceUnreal):
             if not DevicenDisplay.csettings['sound'].get_value(self.name)
             else "")
 
+        # MCP Subsystem
+        no_mcp = (
+            "-nomcp"
+            if not DevicenDisplay.csettings['mcp'].get_value(self.name)
+            else "")
+
         # Loading Screen
         no_loading_screen = (
             "-NoLoadingScreen"
@@ -1171,6 +1183,7 @@ class DevicenDisplay(DeviceUnreal):
             f'{use_all_cores}',           # -useallavailablecores
             f'{no_texture_streaming}',    # -notexturestreaming
             f'{no_sound}',                # -nosound
+            f'{no_mcp}',                  # -nomcp
             f'{no_loading_screen}',       # -NoLoadingScreen
             f'{no_python}',               # -DisablePython
             f'-dc_node={self.name}',      # name of this node in the nDisplay cluster
