Swtichboard: Make soft kill timeout a user parameter (default to 3 seconds), and close nDisplay cluster more gracefully.

The main objective of this change is as a mechanism that allows graceful closure of media I/O, because not doing so sometimes leaves I/O cards in a non-closed state, and they may e.g. keep SDI outputs alive with the last frame dislpayed, which is not always desired.

Added a new FloatSetting property class to hold float settings (used for the timeout in seconds value), which is now used for Unreal and nDIsplay devices when issuing the kill command.

In addition, resurrected the cluster event to shut down the cluster, which ensures all nodes receive the quit event at the same time and may allow a graceful shutdown without cluster barrier issues. When the timeout expires, it still issues the normal taskkill command to each node, since the cluster event may not have been successfull, or the shutdown may be taking long.

Notes:

* The cluster event is not issued when closing the nodes individually, only when using the heading or all device global close buttons in the main dialog.
* Because nDisplay will now effectively have up to 3 kill commands (cluster, taskkill, taskkill /f), then the soft timeout applies effectively twice, so we use split half on each.

Also fixed an issue where the launch buttons would be set to the wrong state when LLH executable was not found.

Bumped SBL version to 3.3.0 (from 3.2.x) to ensure the new timeout parameter is taken into account.

#rnx

diff --git a/Engine/Plugins/VirtualProduction/Switchboard/Source/Switchboard/switchboard/config.py b/Engine/Plugins/VirtualProduction/Switchboard/Source/Switchboard/switchboard/config.py
--- a/Engine/Plugins/VirtualProduction/Switchboard/Source/Switchboard/switchboard/config.py
+++ b/Engine/Plugins/VirtualProduction/Switchboard/Source/Switchboard/switchboard/config.py
@@ -945,6 +945,66 @@ class IntSetting(Setting):
             widget.setCursorPosition(0)
 
 
+class FloatSetting(Setting):
+    '''
+    A UI-displayable Setting for storing and modifying a float value.
+    '''
+
+    def __init__(
+        self,
+        *args,
+        is_read_only: bool = False,
+        **kwargs
+    ):
+        '''
+        Create a new FloatSetting object.
+
+        Args:
+            is_read_only    : Whether to make entry field editable or not.
+        '''
+        super().__init__(*args, **kwargs)
+
+        self.is_read_only = is_read_only
+
+    def _create_widgets(
+            self, override_device_name: Optional[str] = None) \
+            -> QtWidgets.QLineEdit:
+        line_edit = QtWidgets.QLineEdit()
+        if self.tool_tip:
+            line_edit.setToolTip(self.tool_tip)
+        line_edit.setValidator(QtGui.QDoubleValidator())
+
+        value = str(self.get_value(override_device_name))
+        line_edit.setText(value)
+        line_edit.setCursorPosition(0)
+        line_edit.setReadOnly(self.is_read_only)
+
+        self.set_widget(
+            widget=line_edit, override_device_name=override_device_name)
+
+        line_edit.editingFinished.connect(
+            lambda line_edit=line_edit,
+            override_device_name=override_device_name:
+                self._on_widget_value_changed(
+                    float(line_edit.text()),
+                    override_device_name=override_device_name))
+
+        return line_edit
+
+    def _on_setting_changed(
+            self, new_value: float,
+            override_device_name: Optional[str] = None):
+
+        widget = self.get_widget(override_device_name=override_device_name)
+        if not widget:
+            return
+
+        old_value = float(widget.text())
+        if new_value != old_value:
+            widget.setText(str(new_value))
+            widget.setCursorPosition(0)
+
+
 class StringSetting(Setting):
     '''
     A UI-displayable Setting for storing and modifying a string value.
diff --git a/Engine/Plugins/VirtualProduction/Switchboard/Source/Switchboard/switchboard/devices/device_base.py b/Engine/Plugins/VirtualProduction/Switchboard/Source/Switchboard/switchboard/devices/device_base.py
--- a/Engine/Plugins/VirtualProduction/Switchboard/Source/Switchboard/switchboard/devices/device_base.py
+++ b/Engine/Plugins/VirtualProduction/Switchboard/Source/Switchboard/switchboard/devices/device_base.py
@@ -437,6 +437,21 @@ class Device(QtCore.QObject):
         '''
         pass
 
+    @classmethod
+    def close_all_devices(cls, devices):
+        '''
+        Handle closing all devices of this type when "Stop all" button is pressed.
+        Override this method to customize behavior for specific device types.
+        Default implementation closes each device individually.
+        '''
+        for device in devices:
+            try:
+                open_button = device.widget.open_button
+                if open_button.isEnabled() and open_button.isChecked():
+                    device.widget._close()
+            except Exception:
+                pass
+
     @classmethod
     def added_device(cls, device):
         ''' DeviceManager informing that this plugin device has been created and added.
diff --git a/Engine/Plugins/VirtualProduction/Switchboard/Source/Switchboard/switchboard/devices/ndisplay/plugin_ndisplay.py b/Engine/Plugins/VirtualProduction/Switchboard/Source/Switchboard/switchboard/devices/ndisplay/plugin_ndisplay.py
--- a/Engine/Plugins/VirtualProduction/Switchboard/Source/Switchboard/switchboard/devices/ndisplay/plugin_ndisplay.py
+++ b/Engine/Plugins/VirtualProduction/Switchboard/Source/Switchboard/switchboard/devices/ndisplay/plugin_ndisplay.py
@@ -2014,6 +2014,89 @@ class DevicenDisplay(DeviceUnreal):
         except Exception:
             LOGGER.warning("Could not soft kill the cluster")
 
+    # Class variables for managing non-blocking cluster shutdown
+    _cluster_shutdown_timer = None
+    _devices_pending_shutdown = []
+
+    @classmethod
+    def close_all_devices(cls, devices):
+        '''
+        Override to use nDisplay's cluster-aware shutdown with timeout fallback.
+        First tries cluster soft kill, then falls back to individual closes after timeout.
+        Uses half the configured timeout for each phase to keep total time consistent.
+        '''
+        if not devices:
+            return
+            
+        full_timeout = DeviceUnreal.csettings['soft_kill_timeout'].get_value()
+        cluster_timeout = full_timeout / 2.0  # Half timeout for cluster phase
+        
+        try:
+            # Issue cluster command, which is preferred because it issues the quit command at the same time on all
+            # nodes which helps normal shutdown without thread barrier related deadlocks.
+            cls.soft_kill_cluster(devices)
+            
+            # Since the cluster command may not succeed, after a timeout we issue the regular kill command via 
+            # SwitchboardListener on the nodes that are still running.
+
+            cls._devices_pending_shutdown = devices.copy()
+            
+            # Cancel any existing timeout timer
+            if cls._cluster_shutdown_timer:
+                cls._cluster_shutdown_timer.stop()
+                cls._cluster_shutdown_timer = None
+            
+            # Create new timer for the fallback kill command - use half the timeout
+            cls._cluster_shutdown_timer = QtCore.QTimer()
+            cls._cluster_shutdown_timer.timeout.connect(cls._on_cluster_shutdown_timeout)
+            cls._cluster_shutdown_timer.setSingleShot(True)
+            cls._cluster_shutdown_timer.start(int(cluster_timeout * 1000))  # Convert to ms
+            
+            # Set devices to CLOSING state (disables buttons and shows proper state)
+            for device in devices:
+                try:
+                    device.status = DeviceStatus.CLOSING
+                except Exception:
+                    pass
+                    
+        except Exception:
+            LOGGER.warning("Cluster event for shutdown could not be issued, falling back to immediate device close")
+
+            # Issue an immediate regular device shutdown if the cluster command fails to execute.
+            cls._close_nodes_individually(devices)
+
+    @classmethod
+    def _on_cluster_shutdown_timeout(cls):
+        '''Called after timeout - check if devices are still running and force close if needed'''
+        if not cls._devices_pending_shutdown:
+            return
+                    
+        # Check which devices are still running and force close them
+        devices_still_running = []
+        for device in cls._devices_pending_shutdown:
+            # Check if device is still in CLOSING state (meaning cluster shutdown didn't complete)
+            if device.status == DeviceStatus.CLOSING:
+                devices_still_running.append(device)
+        
+        if devices_still_running:            
+            cls._close_nodes_individually(devices_still_running)
+        
+        cls._devices_pending_shutdown.clear()
+        cls._cluster_shutdown_timer = None
+
+    @classmethod
+    def _close_nodes_individually(cls, devices):
+        '''Force close individual devices using the standard close mechanism with remaining timeout'''
+        full_timeout = DeviceUnreal.csettings['soft_kill_timeout'].get_value()
+        remaining_timeout = full_timeout / 2.0  # Use remaining half of timeout
+        
+        for device in devices:
+            try:
+                # Call device.close() directly with the remaining timeout
+                device.close(soft_kill_timeout=remaining_timeout)
+            except Exception:
+                pass
+
     # ~ Begin Device Interface
     @classmethod
     def all_devices_added(cls):
diff --git a/Engine/Plugins/VirtualProduction/Switchboard/Source/Switchboard/switchboard/devices/unreal/plugin_unreal.py b/Engine/Plugins/VirtualProduction/Switchboard/Source/Switchboard/switchboard/devices/unreal/plugin_unreal.py
--- a/Engine/Plugins/VirtualProduction/Switchboard/Source/Switchboard/switchboard/devices/unreal/plugin_unreal.py
+++ b/Engine/Plugins/VirtualProduction/Switchboard/Source/Switchboard/switchboard/devices/unreal/plugin_unreal.py
@@ -36,7 +36,7 @@ import switchboard.switchboard_widgets as sb_widgets
 from switchboard import message_protocol, switchboard_application, ugs_utils
 from switchboard.config import (
     CONFIG, DEFAULT_MAP_TEXT, ENABLE_UGS_SUPPORT, SETTINGS, BoolSetting, DirectoryPathSetting,
-    EngineSyncMethod, FilePathSetting, IntSetting, MultiOptionSetting, OptionSetting, Setting,
+    EngineSyncMethod, FilePathSetting, FloatSetting, IntSetting, MultiOptionSetting, OptionSetting, Setting,
     StringListSetting, StringSetting, migrate_comma_separated_string_to_list)
 from switchboard.devices.device_base import Device, DeviceStatus, PluginHeaderWidgets
 from switchboard.devices.device_widget_base import DeviceAutoJoinMUServerUI, DeviceWidget
@@ -793,7 +793,18 @@ class DeviceUnreal(Device):
                 "The name of your project's editor target. Typically $(ProjectName)Editor (which is inferred "
                 "automatically). This override only needs to be set if deviating from the standard convention."),
             category='UE Settings',
-        )
+        ),
+        'soft_kill_timeout': FloatSetting(
+            attr_name='soft_kill_timeout',
+            nice_name='Soft Kill Timeout',
+            value=3.0,
+            tool_tip=(
+                'Timeout in seconds for graceful process shutdown before force termination. '
+                'Applies to Unreal and nDisplay devices. '
+                'Processes have this amount of time to save state and clean up before being forcefully killed.'
+            ),
+            category="General Settings",
+        ),
     }
 
     if ENABLE_UGS_SUPPORT:
@@ -2003,12 +2014,16 @@ class DeviceUnreal(Device):
 
         return puuid
 
-    def close(self, force=False):
+    def close(self, force=False, soft_kill_timeout=None):
         # This call only refers to "unreal" programs.
         unreal_puuids = self.program_start_queue.running_puuids_named(UnrealJobs.Unreal.value)
 
+        # Get soft kill timeout for this device (use provided value or setting)
+        if soft_kill_timeout is None:
+            soft_kill_timeout = DeviceUnreal.csettings['soft_kill_timeout'].get_value(self.name)
+
         for unreal_puuid in unreal_puuids:
-            _, msg = message_protocol.create_kill_process_message(unreal_puuid)
+            _, msg = message_protocol.create_kill_process_message(unreal_puuid, soft_kill_timeout)
             self.unreal_client.send_message(msg)
 
         if not len(unreal_puuids):
@@ -2482,14 +2497,14 @@ class DeviceUnreal(Device):
             LOGGER.error(f"Could not start {program_name}: {message['error']}")
 
             if program_name == UnrealJobs.Sync.value or program_name.startswith(UnrealJobs.BuildPrefix.value):
+                # Sync/Build failures may affect multiple devices
                 for device in self.devices_sharing_workspace():
                     device.status = DeviceStatus.CLOSED
 
                     # This has the effect of hiding building/syncing status.
                     device.project_changelist = device.project_changelist
-            elif program_name == UnrealJobs.Unreal.value:
-                self.status = DeviceStatus.CLOSED
-            elif program_name == UnrealJobs.Retrieve.value:
+            else:
+                # Any other program failure sets this device to CLOSED
                 self.status = DeviceStatus.CLOSED
 
             return
diff --git a/Engine/Plugins/VirtualProduction/Switchboard/Source/Switchboard/switchboard/devices/unreal/version_helpers.py b/Engine/Plugins/VirtualProduction/Switchboard/Source/Switchboard/switchboard/devices/unreal/version_helpers.py
--- a/Engine/Plugins/VirtualProduction/Switchboard/Source/Switchboard/switchboard/devices/unreal/version_helpers.py
+++ b/Engine/Plugins/VirtualProduction/Switchboard/Source/Switchboard/switchboard/devices/unreal/version_helpers.py
@@ -3,7 +3,7 @@
 import typing
 
 
-LISTENER_COMPATIBLE_VERSION = (3,2)
+LISTENER_COMPATIBLE_VERSION = (3,3)
 
 # Version where redeploy support was initially introduced.
 LISTENER_MIN_REDEPLOY_VERSION = (1,4,0)
diff --git a/Engine/Plugins/VirtualProduction/Switchboard/Source/Switchboard/switchboard/message_protocol.py b/Engine/Plugins/VirtualProduction/Switchboard/Source/Switchboard/switchboard/message_protocol.py
--- a/Engine/Plugins/VirtualProduction/Switchboard/Source/Switchboard/switchboard/message_protocol.py
+++ b/Engine/Plugins/VirtualProduction/Switchboard/Source/Switchboard/switchboard/message_protocol.py
@@ -97,9 +97,11 @@ def create_start_process_message(
     return (cmd_id, message)
 
 
-def create_kill_process_message(program_id):
+def create_kill_process_message(program_id, soft_kill_timeout=None):
     cmd_id = uuid.uuid4()
     kill_cmd = {'command': 'kill', 'id': str(cmd_id), 'uuid': str(program_id)}
+    if soft_kill_timeout is not None:
+        kill_cmd['soft_kill_timeout'] = soft_kill_timeout
     message = json.dumps(kill_cmd).encode() + b'\x00'
     return (cmd_id, message)
 
diff --git a/Engine/Plugins/VirtualProduction/Switchboard/Source/Switchboard/switchboard/switchboard_dialog.py b/Engine/Plugins/VirtualProduction/Switchboard/Source/Switchboard/switchboard/switchboard_dialog.py
--- a/Engine/Plugins/VirtualProduction/Switchboard/Source/Switchboard/switchboard/switchboard_dialog.py
+++ b/Engine/Plugins/VirtualProduction/Switchboard/Source/Switchboard/switchboard/switchboard_dialog.py
@@ -1399,7 +1399,21 @@ class SwitchboardDialog(QtCore.QObject):
 
     def launch_all_button_clicked(self, button_state):
         devices = self.device_manager.devices()
-        self.set_device_launch_state(devices, button_state)
+        if button_state:
+            # Opening devices.
+            self.set_device_launch_state(devices, button_state)
+        else:
+            # Closing devices - group by type and use device specific close_all_devices()
+            devices_by_type = {}
+            for device in devices:
+                device_type = type(device)
+                if device_type not in devices_by_type:
+                    devices_by_type[device_type] = []
+                devices_by_type[device_type].append(device)
+            
+            # Call close_all_devices() for each device type
+            for device_class, device_list in devices_by_type.items():
+                device_class.close_all_devices(device_list)
 
     def settings_button_clicked(self, button_state):
         self.menu_update_settings()
@@ -1548,7 +1562,13 @@ class SwitchboardDialog(QtCore.QObject):
 
     def on_all_plugin_devices_open_toggled(self, plugin_name, button_state):
         devices = self.device_manager.devices_of_type(plugin_name)
-        self.set_device_launch_state(devices, button_state)
+        
+        # Use device specific close_all_devices() method when Stopping.
+        if not button_state and devices:
+            device_class = type(devices[0])
+            device_class.close_all_devices(devices)
+        else:
+            self.set_device_launch_state(devices, button_state)
 
     def multiuser_session_name(self):
         return self._multiuser_session_name
diff --git a/Engine/Source/Programs/SwitchboardListener/SblCore/Internal/SwitchboardListenerVersion.h b/Engine/Source/Programs/SwitchboardListener/SblCore/Internal/SwitchboardListenerVersion.h
--- a/Engine/Source/Programs/SwitchboardListener/SblCore/Internal/SwitchboardListenerVersion.h
+++ b/Engine/Source/Programs/SwitchboardListener/SblCore/Internal/SwitchboardListenerVersion.h
@@ -3,5 +3,5 @@
 #pragma once
 
 #define SBLISTENER_VERSION_MAJOR 3
-#define SBLISTENER_VERSION_MINOR 2
-#define SBLISTENER_VERSION_PATCH 2
+#define SBLISTENER_VERSION_MINOR 3
+#define SBLISTENER_VERSION_PATCH 0
diff --git a/Engine/Source/Programs/SwitchboardListener/SblCore/Private/SwitchboardListener.cpp b/Engine/Source/Programs/SwitchboardListener/SblCore/Private/SwitchboardListener.cpp
--- a/Engine/Source/Programs/SwitchboardListener/SblCore/Private/SwitchboardListener.cpp
+++ b/Engine/Source/Programs/SwitchboardListener/SblCore/Private/SwitchboardListener.cpp
@@ -1161,7 +1161,7 @@ bool FSwitchboardListener::Task_KillProcess(const FSwitchboardKillTask& KillTask
 	MessageFuture.Future = Async(EAsyncExecution::ThreadPool, [this, KillTask, Process, FlipModeMonitor, UUID]() {
 		SWITCHBOARD_TRACE_CPUPROFILER_EVENT_SCOPE_STR("FSwitchboardListener::Task_KillProcess future closure");
 
-		const float SoftKillTimeout = 2.0f;
+		const float SoftKillTimeout = FMath::Max(0.0f, KillTask.SoftKillTimeout);
 
 		const bool bKilledProcess = KillProcessNow(Process.Get(), SoftKillTimeout);
 		KillProcessNow(FlipModeMonitor.Get(), SoftKillTimeout);
diff --git a/Engine/Source/Programs/SwitchboardListener/SblCore/Private/SwitchboardProtocol.cpp b/Engine/Source/Programs/SwitchboardListener/SblCore/Private/SwitchboardProtocol.cpp
--- a/Engine/Source/Programs/SwitchboardListener/SblCore/Private/SwitchboardProtocol.cpp
+++ b/Engine/Source/Programs/SwitchboardListener/SblCore/Private/SwitchboardProtocol.cpp
@@ -303,8 +303,15 @@ FCreateTaskResult CreateTaskFromCommand(const FString& InCommand, const FIPv4End
 		FGuid ProgramID;
 		if (UUIDField.IsValid() && FGuid::Parse(UUIDField->AsString(), ProgramID))
 		{
+			// Extract optional soft kill timeout, use default if not provided
+			TUniquePtr<FSwitchboardKillTask> KillTask = MakeUnique<FSwitchboardKillTask>(MessageID, InEndpoint, ProgramID);
+			if (TSharedPtr<FJsonValue> TimeoutField = JsonData->TryGetField(TEXT("soft_kill_timeout")))
+			{
+				KillTask->SoftKillTimeout = TimeoutField->AsNumber();
+			}
+
 			Result.Status = ECreateTaskStatus::Success;
-			Result.Task = MakeUnique<FSwitchboardKillTask>(MessageID, InEndpoint, ProgramID);
+			Result.Task = MoveTemp(KillTask);
 			return Result;
 		}
 	}
diff --git a/Engine/Source/Programs/SwitchboardListener/SblCore/Private/SwitchboardTasks.h b/Engine/Source/Programs/SwitchboardListener/SblCore/Private/SwitchboardTasks.h
--- a/Engine/Source/Programs/SwitchboardListener/SblCore/Private/SwitchboardTasks.h
+++ b/Engine/Source/Programs/SwitchboardListener/SblCore/Private/SwitchboardTasks.h
@@ -186,12 +186,17 @@ struct FSwitchboardStartTask : public FSwitchboardTask
 
 struct FSwitchboardKillTask : public FSwitchboardTask
 {
-	FSwitchboardKillTask(const FGuid& InTaskId, const FIPv4Endpoint& InEndpoint, const FGuid& InProgramID)
+	FSwitchboardKillTask(const FGuid& InTaskId, const FIPv4Endpoint& InEndpoint, const FGuid& InProgramID, float InSoftKillTimeout = -1.0f)
 		: FSwitchboardTask{ ESwitchboardTaskType::Kill, InTaskId, InEndpoint}
 		, ProgramID(InProgramID)
+		, SoftKillTimeout(InSoftKillTimeout)
 	{}
 
-	FGuid ProgramID; // unique ID of process to kill
+	// unique ID of process to kill
+	FGuid ProgramID;
+
+	/** Time after a soft termination signal before issuing a hard termination signal. */
+	float SoftKillTimeout = 2.0f;
 
 	//~ Begin FSwitchboardTask interface
 	static constexpr const TCHAR* CommandName = TEXT("kill");
@@ -199,6 +204,8 @@ struct FSwitchboardKillTask : public FSwitchboardTask
 
 	virtual uint32 GetEquivalenceHash() const override
 	{
+		// Ignore the SoftKillTimeout for this.  A new kill command of the same program with a different 
+		// timeout should be considered a duplicate.
 		return HashCombine(FSwitchboardTask::GetEquivalenceHash(), GetTypeHash(ProgramID));
 	}
 	//~ End FSwitchboardTask interface
