Switchboard: Add experimental "fast sync" to sbl_helper.py.

When the currently synced changelist is known, and we are syncing ahead to a future changelist, pass a revision range to the sync command.

A corresponding "Use Fast Sync" setting has been added to Switchboard's "Source Control Settings" section, disabled by default.

In order to mitigate additional risk of the workspace ending up in an inconsistent state if UGS and Switchboard are used interchangably, we now compare the /.ugs/state.json "LastSyncTime" field to the timestamp in .sblsync.json, treating a more recent UGS sync as an invalidating condition.

Similarly, we have to take the last synced .uproject into account: we conservatively disable fast sync when it changes, in case the working set of AdditionalPathsToSync changes, and the "current" CL for those paths may be inaccurate.

Finally, we also now look for the presence of "WARNING" or "ERROR" in the sync task output, and highlight those entries accordingly in the Switchboard log.

diff --git a/Engine/Plugins/VirtualProduction/Switchboard/Source/Switchboard/sbl_helper.py b/Engine/Plugins/VirtualProduction/Switchboard/Source/Switchboard/sbl_helper.py
--- a/Engine/Plugins/VirtualProduction/Switchboard/Source/Switchboard/sbl_helper.py
+++ b/Engine/Plugins/VirtualProduction/Switchboard/Source/Switchboard/sbl_helper.py
@@ -249,6 +249,8 @@ def p4_get_preview_sync_files(
             if record[b'data'].endswith(b'up-to-date.\n'):
                 logging.info(record[b'data'].decode().rstrip())
                 continue
+            elif record[b'data'].endswith(b'no such file(s).\n'):
+                continue  # Harmless, and common with revision ranges.
 
         logging.warning('p4_get_preview_sync_files(): Unhandled record: '
                         f'{record}')
@@ -487,6 +489,7 @@ class SbListenerHelper:
         self.log_string = io.StringIO()
 
         self.uproj_path: Optional[pathlib.Path] = None
+        self.project_dir: Optional[pathlib.Path] = None
         self.engine_dir: Optional[pathlib.Path] = None
         self.additional_paths_to_sync: list[str] = []
         self.sync_engine_cl: Optional[int] = None
@@ -495,6 +498,12 @@ class SbListenerHelper:
         self.clobber_project: bool = False
         self.generate_after_sync: bool = False
         self.ugs_versioning: bool = False
+        self.no_fast_sync: bool = False
+
+        # UGS CLI options
+        self.use_ugs: bool = False
+        self.sync_pcbs: bool = False
+        self.ugs_lib_dir: Optional[pathlib.Path] = None
 
         self.dry_run: bool = False
         self.p4port: Optional[str] = None
@@ -504,6 +513,9 @@ class SbListenerHelper:
         self.p4kwargs: dict[str, Any] = {}
         self.p4clientspec: dict[str, str] = {}
 
+        self.engine_state: Optional[SbListenerHelper.EngineSyncState] = None
+        self.project_state: Optional[SbListenerHelper.SyncState] = None
+
         if UGS.SyncFilters.supported():
             self.ugs_filters = UGS.SyncFilters()
         else:
@@ -568,6 +580,8 @@ class SbListenerHelper:
                                  help='Override noclobber for project files')
         sync_parser.add_argument('--ugs-versioning', action='store_true',
                                  help='Maintain depot CompatibleChangelist (e.g. hotfix stream)')
+        sync_parser.add_argument('--no-fast-sync', action='store_true',
+                                 help='Disable "fast sync" (revision range) optimization')
 
         sync_parser.add_argument('-n', '--dry-run', action='store_true',
                                  help='Preview sync; do not update files')
@@ -618,6 +632,32 @@ class SbListenerHelper:
             handlers=handlers
         )
 
+        logging.debug(f'Options: {options}')
+
+        if options.project:
+            try:
+                self.uproj_path = options.project.resolve(strict=True)
+            except FileNotFoundError:
+                self.parser.error('argument --project: FileNotFoundError: '
+                                  f'{options.project}')
+
+            try:
+                _ = json.load(self.uproj_path.open())
+            except json.JSONDecodeError:
+                self.parser.error('argument --project: unable to parse JSON'
+                                  f'\n\n{traceback.format_exc()}')
+
+        self.engine_dir = self.find_engine_dir(options)
+
+        if not self.engine_dir and not self.uproj_path:
+            self.parser.error('Neither engine nor project path were provided')
+
+        self.project_dir = self.uproj_path.parent if self.uproj_path else None
+
+        # Attempt to read the notionally synced changelists, as written during the last sync.
+        self.engine_state = self.read_engine_sync_state()
+        self.project_state = self.read_project_sync_state()
+
         if options.action == Subcommands.Sync:
             return self.run_sync(options)
         elif options.action == Subcommands.SyncStatus:
@@ -630,8 +670,10 @@ class SbListenerHelper:
         options: argparse.Namespace,
         *,
         engine_dir: pathlib.Path | None = None,
+        prev_engine_cl: Optional[int] = None,
         sync_engine_cl: Optional[int] = None,
         project_dir: pathlib.Path | None = None,
+        prev_project_cl: Optional[int] = None,
         sync_project_cl: Optional[int] = None,
         port: Optional[str] = None,
         user: Optional[str] = None,
@@ -686,8 +728,12 @@ class SbListenerHelper:
             engine_paths.append(f'{engine_dir.parent / "*"}')
             engine_paths.append(f'{engine_dir / "..."}')
 
+            revspec = f'@{sync_engine_cl}'
+            if prev_engine_cl and prev_engine_cl <= sync_engine_cl:
+                revspec = f'@{prev_engine_cl},{sync_engine_cl}'
+
             pathspecs.extend(
-                [f'{path}@{sync_engine_cl}' for path in engine_paths])
+                [f'{path}{revspec}' for path in engine_paths])
 
         if sync_project_cl:
             assert project_dir
@@ -696,8 +742,12 @@ class SbListenerHelper:
             for path in self.additional_paths_to_sync:
                 project_paths.append(f'//{self.p4client}{path}')
 
+            revspec = f'@{sync_project_cl}'
+            if prev_project_cl and prev_project_cl <= sync_project_cl:
+                revspec = f'@{prev_project_cl},{sync_project_cl}'
+
             pathspecs.extend(
-                [f'{path}@{sync_project_cl}' for path in project_paths])
+                [f'{path}{revspec}' for path in project_paths])
 
         if len(pathspecs) < 1:
             self.parser.error('Nothing specified to sync')
@@ -780,11 +830,45 @@ class SbListenerHelper:
                 client=self.p4client
             )
         else:
-            project_dir = self.uproj_path.parent if self.uproj_path else None
+            prev_engine_cl=self.engine_state.changelist if self.engine_state else None
+            prev_project_cl=self.project_state.changelist if self.project_state else None
+
+            # If the user has synced the workspace with UGS recently, flag it
+            if self.engine_dir:
+                ugs_state_path = self.engine_dir.parent / '.ugs' / 'state.json'
+                if ugs_state_path.exists():
+                    with open(ugs_state_path, 'rt', encoding='utf-8') as f:
+                        try:
+                            ugs_state = json.load(f)
+                            ugs_last_sync_time = ugs_state['LastSyncTime']
+                            ugs_last_sync_datetime = datetime.datetime.fromisoformat(ugs_last_sync_time)
+                            if self.engine_state and self.engine_state.sync_time < ugs_last_sync_datetime:
+                                logging.warning('Disabling fast sync because UGS last synced the workspace')
+                                self.no_fast_sync = True
+                        except Exception:
+                            logging.warning('Error comparing .ugs/state.json LastSyncTime')
+
+            if self.no_fast_sync:
+                logging.info('Fast sync is disabled')
+                prev_engine_cl = None
+                prev_project_cl = None
+            else:
+                if (not self.engine_state
+                    or not self.engine_state.last_synced_project
+                    or self.engine_state.last_synced_project != self.uproj_path
+                ):
+                    logging.warning('Project fast sync disabled because last synced project differs')
+                    if self.engine_state:
+                        logging.warning(f'({self.engine_state.last_synced_project})')
+                    prev_project_cl = None
 
             sync_workload, buildver_info = self._get_sync_workload(options,
-                engine_dir=self.engine_dir, sync_engine_cl=self.sync_engine_cl,
-                project_dir=project_dir, sync_project_cl=self.sync_project_cl,
+                engine_dir=self.engine_dir,
+                prev_engine_cl=prev_engine_cl,
+                sync_engine_cl=self.sync_engine_cl,
+                project_dir=self.project_dir,
+                prev_project_cl=prev_project_cl,
+                sync_project_cl=self.sync_project_cl,
                 **self.p4kwargs)
 
             # User may have requested to write the pending sync files and revisions to
@@ -870,72 +954,105 @@ class SbListenerHelper:
         changelist: int
         sync_time: datetime.datetime
 
-    def _save_sync_state_to_file(self, sync_state: SyncState, path: pathlib.Path):
-        data = {
-            'changelist': sync_state.changelist,
-            'sync_time': sync_state.sync_time.isoformat(),
-        }
+        def save_to_file(self, path: pathlib.Path):
+            data = self._get_fields()
 
-        try:
-            with open(path, 'wt', encoding='utf-8') as f:
-                json.dump(data, f, indent=4)
-        except Exception as exc:
-            logging.error(f'Unable to write sync state to {path}', exc_info=exc)
-            return
+            try:
+                with open(path, 'wt', encoding='utf-8') as f:
+                    json.dump(data, f, indent=4)
+            except Exception as exc:
+                logging.error(f'Unable to write sync state to {path}', exc_info=exc)
+                return
+
+            logging.debug(f'Sync state written to {path}: {data}')
+
+        def _get_fields(self) -> dict[str, Any]:
+            return {
+                'changelist': self.changelist,
+                'sync_time': self.sync_time.isoformat(),
+            }
+
+        @classmethod
+        def load_from_file(cls, path: pathlib.Path) -> SbListenerHelper.SyncState | None:
+            logging.debug(f'Reading SyncState from {path}')
+            fields = cls._read_fields_from_file(path)
+            logging.debug(f'SyncState: {fields}')
+
+            if not fields:
+                return None
 
-        logging.debug(f'Sync state written to {path}: {data}')
+            return SbListenerHelper.SyncState(
+                changelist=int(fields['changelist']),
+                sync_time=datetime.datetime.fromisoformat(fields['sync_time'])
+            )
 
-    def _load_sync_state_from_file(self, path: pathlib.Path) -> SyncState | None:
-        if not path.exists():
-            logging.debug(f'No sync state located at {path}')
-            return None
+        @classmethod
+        def _read_fields_from_file(cls, path: pathlib.Path) -> dict[str, Any] | None:
+            if not path.exists():
+                logging.debug(f'No sync state located at {path}')
+                return None
 
-        try:
-            with open(path, 'r', encoding='utf-8') as f:
-                state_json: dict = json.load(f)
-                logging.debug(f'Sync state loaded from {path}: {state_json}')
-        except Exception as exc:
-            logging.error(f'Unable to read sync state from {path}', exc_info=exc)
-            return None
+            try:
+                with open(path, 'r', encoding='utf-8') as f:
+                    state_json: dict = json.load(f)
+                    return state_json
+            except Exception as exc:
+                logging.error(f'Unable to read sync state from {path}', exc_info=exc)
+                return None
 
-        return SbListenerHelper.SyncState(
-            changelist=int(state_json['changelist']),
-            sync_time=datetime.datetime.fromisoformat(state_json['sync_time'])
-        )
+    @dataclass
+    class EngineSyncState(SyncState):
+        last_synced_project: pathlib.Path | None
+
+        @classmethod
+        def load_from_file(cls, path: pathlib.Path) -> SbListenerHelper.EngineSyncState | None:
+            logging.debug(f'Reading EngineSyncState from {path}')
+            fields = cls._read_fields_from_file(path)
+            logging.debug(f'EngineSyncState: {fields}')
+
+            if not fields:
+                return None
+
+            return SbListenerHelper.EngineSyncState(
+                changelist=int(fields['changelist']),
+                sync_time=datetime.datetime.fromisoformat(fields['sync_time']),
+                last_synced_project=pathlib.Path(fields['last_synced_project']),
+            )
 
-    def read_engine_sync_state(self) -> SyncState | None:
+        def _get_fields(self) -> dict[str, Any]:
+            return {
+                **super()._get_fields(),
+                'last_synced_project': str(self.last_synced_project),
+            }
+
+    def read_engine_sync_state(self) -> EngineSyncState | None:
         if self.engine_sync_state_path:
-            return self._load_sync_state_from_file(self.engine_sync_state_path)
+            return SbListenerHelper.EngineSyncState.load_from_file(self.engine_sync_state_path)
+        else:
+            logging.warning(f'{self.STATE_FILE} not available for {self.engine_sync_state_path}')
 
     def read_project_sync_state(self) -> SyncState | None:
         if self.project_sync_state_path:
-            return self._load_sync_state_from_file(self.project_sync_state_path)
+            return SbListenerHelper.SyncState.load_from_file(self.project_sync_state_path)
+        else:
+            logging.warning(f'{self.STATE_FILE} not available for {self.project_sync_state_path}')
 
     def write_sync_state(self):
-        sync_time = datetime.datetime.now()
+        sync_time = datetime.datetime.now(datetime.timezone.utc)
 
         if self.sync_engine_cl and self.engine_sync_state_path:
-            engine_state = SbListenerHelper.SyncState(self.sync_engine_cl, sync_time)
-            self._save_sync_state_to_file(engine_state, self.engine_sync_state_path)
+            engine_state = SbListenerHelper.EngineSyncState(self.sync_engine_cl, sync_time, self.uproj_path)
+            engine_state.save_to_file(self.engine_sync_state_path)
 
         if self.sync_project_cl and self.project_sync_state_path:
             project_state = SbListenerHelper.SyncState(self.sync_project_cl, sync_time)
-            self._save_sync_state_to_file(project_state, self.project_sync_state_path)
+            project_state.save_to_file(self.project_sync_state_path)
 
     def run_syncstatus(self, options: argparse.Namespace) -> int:
         self.check_p4_options(options)
 
-        if not self.engine_dir and not self.uproj_path:
-            self.parser.error('Neither engine nor project path were provided')
-
-        project_dir = self.uproj_path.parent if self.uproj_path else None
-
-        # Read the notionally synced changelists, as written during the last attempted sync.
-        engine_state = self.read_engine_sync_state()
-        project_state = self.read_project_sync_state()
-
         def guess_have_cl(path: pathlib.Path, client: str | None = None) -> int | None:
-            logging.debug(f'{self.STATE_FILE} not available for {path}; falling back to `p4 cstat #have`')
+            logging.debug(f'Falling back to `p4 cstat #have` for {path}')
 
             if not path.is_dir():
                 return None
@@ -966,26 +1083,26 @@ class SbListenerHelper:
             if eng_detected_ws := p4_utils.p4_get_var('P4CLIENT', self.engine_dir):
                 result['engine_detected_ws'] = eng_detected_ws
 
-            if engine_state:
-                self.sync_engine_cl = engine_state.changelist
-                result['engine_cl'] = engine_state.changelist
-                result['engine_sync_time'] = engine_state.sync_time.isoformat()
+            if self.engine_state:
+                self.sync_engine_cl = self.engine_state.changelist
+                result['engine_cl'] = self.engine_state.changelist
+                result['engine_sync_time'] = self.engine_state.sync_time.isoformat()
             else:
                 guess_eng_cl = guess_have_cl(self.engine_dir)
                 self.sync_engine_cl = guess_eng_cl
                 result['engine_cl'] = guess_eng_cl
                 result['engine_sync_time'] = '(Unknown)'
 
-        if project_dir:
-            if proj_detected_ws := p4_utils.p4_get_var('P4CLIENT', project_dir):
+        if self.project_dir:
+            if proj_detected_ws := p4_utils.p4_get_var('P4CLIENT', self.project_dir):
                 result['project_detected_ws'] = proj_detected_ws
 
-            if project_state:
-                self.sync_project_cl = project_state.changelist
-                result['project_cl'] = project_state.changelist
-                result['project_sync_time'] = project_state.sync_time.isoformat()
+            if self.project_state:
+                self.sync_project_cl = self.project_state.changelist
+                result['project_cl'] = self.project_state.changelist
+                result['project_sync_time'] = self.project_state.sync_time.isoformat()
             else:
-                guess_proj_cl = guess_have_cl(project_dir, self.p4client_project)
+                guess_proj_cl = guess_have_cl(self.project_dir, self.p4client_project)
                 self.sync_project_cl = guess_proj_cl
                 result['project_cl'] = guess_proj_cl
                 result['project_sync_time'] = '(Unknown)'
@@ -1005,22 +1122,26 @@ class SbListenerHelper:
                 if self.engine_dir:
                     os.chdir(self.engine_dir)
                     eng_workload, _ = self._get_sync_workload(options,
-                        engine_dir=self.engine_dir, sync_engine_cl=self.sync_engine_cl,
+                        engine_dir=self.engine_dir,
+                        sync_engine_cl=self.sync_engine_cl,
                         port=self.p4port, user=self.p4user, client=self.p4client)
                     either_dirty = either_dirty or len(eng_workload.sync_files) > 0
 
-                if project_dir:
-                    os.chdir(project_dir)
+                if self.project_dir:
+                    os.chdir(self.project_dir)
                     proj_workload, _ = self._get_sync_workload(options,
-                        project_dir=project_dir, sync_project_cl=self.sync_project_cl,
+                        project_dir=self.project_dir,
+                        sync_project_cl=self.sync_project_cl,
                         port=self.p4port, user=self.p4user, client=self.p4client_project)
                     either_dirty = either_dirty or len(proj_workload.sync_files) > 0
 
                 result['dirty'] = either_dirty
             else:
                 sync_workload, _ = self._get_sync_workload(options,
-                    engine_dir=self.engine_dir, sync_engine_cl=self.sync_engine_cl,
-                    project_dir=project_dir, sync_project_cl=self.sync_project_cl,
+                    engine_dir=self.engine_dir,
+                    sync_engine_cl=self.sync_engine_cl,
+                    project_dir=self.project_dir,
+                    sync_project_cl=self.sync_project_cl,
                     **self.p4kwargs)
 
                 result['dirty'] = len(sync_workload.sync_files) > 0
@@ -1249,21 +1370,6 @@ class SbListenerHelper:
         else:
             self.parser.error('No valid Perforce session found. Make sure you are logged into Perforce.')
 
-        if options.project:
-            try:
-                self.uproj_path = options.project.resolve(strict=True)
-            except FileNotFoundError:
-                self.parser.error('argument --project: FileNotFoundError: '
-                                  f'{options.project}')
-
-            try:
-                _ = json.load(self.uproj_path.open())
-            except json.JSONDecodeError:
-                self.parser.error('argument --project: unable to parse JSON'
-                                  f'\n\n{traceback.format_exc()}')
-
-        self.engine_dir = self.find_engine_dir(options)
-
     def check_sync_options(self, options: argparse.Namespace):
         '''
         Parse `sync` options, caching their values, and ensuring consistency
@@ -1275,6 +1381,11 @@ class SbListenerHelper:
         self.clobber_engine = options.clobber_engine
         self.clobber_project = options.clobber_project
         self.ugs_versioning = options.ugs_versioning
+        self.no_fast_sync = options.no_fast_sync
+
+        # UGS CLI options
+        self.use_ugs = options.use_ugs
+        self.sync_pcbs = options.use_pcbs
         self.ugs_lib_dir = options.ugs_lib_dir
 
         if self.p4client:
@@ -1291,10 +1402,9 @@ class SbListenerHelper:
 
         self.sync_engine_cl = options.engine_cl
         self.sync_project_cl = options.project_cl
-        self.use_ugs = options.use_ugs
+
         # UGS will generate the project files by default (don't duplicate work)
         self.generate_after_sync = options.generate and not self.use_ugs
-        self.sync_pcbs = options.use_pcbs
 
         if self.sync_pcbs and not self.use_ugs:
             self.parser.error('`--use-pcbs` requires `--use-ugs`.')
diff --git a/Engine/Plugins/VirtualProduction/Switchboard/Source/Switchboard/switchboard/devices/unreal/plugin_unreal.py b/Engine/Plugins/VirtualProduction/Switchboard/Source/Switchboard/switchboard/devices/unreal/plugin_unreal.py
--- a/Engine/Plugins/VirtualProduction/Switchboard/Source/Switchboard/switchboard/devices/unreal/plugin_unreal.py
+++ b/Engine/Plugins/VirtualProduction/Switchboard/Source/Switchboard/switchboard/devices/unreal/plugin_unreal.py
@@ -734,6 +734,13 @@ class DeviceUnreal(Device):
             show_ui=True if sys.platform in ('win32', 'linux') else False,  # Only available in select platforms
             category="GPU/CPU Settings",
         ),
+        'use_fast_sync': BoolSetting(
+            attr_name='use_fast_sync',
+            nice_name='Use Fast Sync',
+            value=False,
+            tool_tip='Controls whether Perforce sync operations use the revision range optimization',
+            category="Source Control Settings",
+        ),
         'use_sync_filters': BoolSetting(
             attr_name='use_sync_filters',
             nice_name='Use Sync Filters',
@@ -1738,6 +1745,9 @@ class DeviceUnreal(Device):
         if engine_cl:
             sync_args.append(f'--engine-dir="{engine_dir}"')
 
+        if not DeviceUnreal.csettings['use_fast_sync'].get_value():
+            sync_args.append('--no-fast-sync')
+
         sync_args = ' '.join(sync_args)
 
         if workspace:
@@ -2831,6 +2841,8 @@ class DeviceUnreal(Device):
         stdout = stdoutbytes.decode()
         lines = list(filter(None, stdout.splitlines()))
 
+        already_logged = False
+
         # see if this is an update to the build
         #
         # example lines:
@@ -2863,6 +2875,17 @@ class DeviceUnreal(Device):
 
         elif process['name'] == UnrealJobs.Sync.value:
             for line in lines:
+                already_logged = True
+
+                if 'ERROR' in line:
+                    log_fn = LOGGER.error
+                elif 'WARNING' in line:
+                    log_fn = LOGGER.warning
+                else:
+                    log_fn = LOGGER.debug
+
+                log_fn(f"{self.name} {process['name']}: {line}")
+
                 if 'Progress:' not in line:
                     continue
 
@@ -2881,8 +2904,9 @@ class DeviceUnreal(Device):
             self.device_qt_handler.signal_device_package_update.emit(
                         self, step, sync_progress)
 
-        for line in lines:
-            LOGGER.debug(f"{self.name} {process['name']}: {line}")
+        if not already_logged:
+            for line in lines:
+                LOGGER.debug(f"{self.name} {process['name']}: {line}")
 
     @classmethod
     def devices(
