Switchboard: Fix deserializing pre-existing EngineSyncState .sblsync.json files. Followup to CL 44756904.

diff --git a/Engine/Plugins/VirtualProduction/Switchboard/Source/Switchboard/sbl_helper.py b/Engine/Plugins/VirtualProduction/Switchboard/Source/Switchboard/sbl_helper.py
--- a/Engine/Plugins/VirtualProduction/Switchboard/Source/Switchboard/sbl_helper.py
+++ b/Engine/Plugins/VirtualProduction/Switchboard/Source/Switchboard/sbl_helper.py
@@ -1016,7 +1016,8 @@ class SbListenerHelper:
             return SbListenerHelper.EngineSyncState(
                 changelist=int(fields['changelist']),
                 sync_time=datetime.datetime.fromisoformat(fields['sync_time']),
-                last_synced_project=pathlib.Path(fields['last_synced_project']),
+                last_synced_project=(pathlib.Path(fields['last_synced_project'])
+                                     if 'last_synced_project' in fields else None),
             )
 
         def _get_fields(self) -> dict[str, Any]:
