Switchboard: Replace wmic with powershell because wmic was deprecated in Win 11.

Also improved the detection of the project path from just the name by looking inside *.uprojectdirs like the Editor does.

#rnx

diff --git a/Engine/Plugins/VirtualProduction/Switchboard/Source/Switchboard/switchboard/add_config_dialog.py b/Engine/Plugins/VirtualProduction/Switchboard/Source/Switchboard/switchboard/add_config_dialog.py
--- a/Engine/Plugins/VirtualProduction/Switchboard/Source/Switchboard/switchboard/add_config_dialog.py
+++ b/Engine/Plugins/VirtualProduction/Switchboard/Source/Switchboard/switchboard/add_config_dialog.py
@@ -8,6 +8,7 @@ import subprocess
 import sys
 import threading
 
+from pathlib import Path
 from PySide6 import QtCore
 from PySide6 import QtWidgets
 
@@ -109,8 +110,9 @@ class AddConfigDialog(QtWidgets.QDialog):
 
         self.engine_dir_line_edit = QtWidgets.QLineEdit()
         if os.path.exists(previous_engine_dir): # re-use previous engine dir
-            self.engine_dir_line_edit.setText(previous_engine_dir)
-            self.engine_dir = previous_engine_dir
+            normalized_previous_engine_dir = os.path.normpath(str(Path(previous_engine_dir)))
+            self.engine_dir_line_edit.setText(normalized_previous_engine_dir)
+            self.engine_dir = normalized_previous_engine_dir
         self.engine_dir_line_edit.textChanged.connect(self.on_engine_dir_changed)
         self.engine_dir_browse_button = QtWidgets.QPushButton("Browse")
         self.engine_dir_browse_button.clicked.connect(self.on_browse_engine_dir)
@@ -438,10 +440,10 @@ class AddConfigDialog(QtWidgets.QDialog):
         # Populate the fields that we're updating and were found
 
         if bUpdateEngine and editorfolder:
-            self.engine_dir_line_edit.setText(editorfolder)
+            self.engine_dir_line_edit.setText(str(Path(editorfolder)))
 
         if bUpdateProject and projectpath:
-            self.uproject_line_edit.setText(projectpath)
+            self.uproject_line_edit.setText(str(Path(projectpath)))
 
     def detect_running_projects(self):
         ''' Detects a running UnrealEngine editor and its project
@@ -462,8 +464,7 @@ class AddConfigDialog(QtWidgets.QDialog):
                 # Windows UE Editors will have .exe extension
                 UEname += '.exe'
 
-                # Relying on wmic for this. Another option is to use psutil which is cross-platform, and slower.
-                cmd = f'wmic process where caption="{UEname}" get commandline'
+                cmd = ['powershell', '-Command', f'Get-WmiObject Win32_Process -Filter "Name=\'{UEname}\'" | Select-Object -ExpandProperty CommandLine']
 
                 try:
                     commandlines = subprocess.check_output(
@@ -491,8 +492,8 @@ class AddConfigDialog(QtWidgets.QDialog):
                     if len(argv) < 2:
                         continue
 
-                    editorpath = os.path.normpath(argv[0])
-                    projectpath = os.path.normpath(argv[1])
+                    editorpath = str(Path(argv[0]))
+                    projectpath = str(Path(argv[1]))
 
                     if editorpath.lower().endswith(UEname.lower()) and os.path.exists(editorpath):
                         editors.append(editorpath)
@@ -502,10 +503,61 @@ class AddConfigDialog(QtWidgets.QDialog):
                     if projectpath.lower().endswith('.uproject') and os.path.exists(projectpath):
                         projects.append(projectpath)
                     else:
-                        projects.append('')
+                        # Try to find project by name using *.uprojectdirs files
+                        found_project = self._find_project_by_name(editorpath, projectpath)
+                        projects.append(found_project if found_project else '')
 
         return editors,projects
 
+    def _find_project_by_name(self, editorpath, project_name):
+        ''' Finds the .uproject path by name using *.uprojectdirs files in the given editor's root folder.
+        '''
+        try:
+            # Look for .uprojectdirs files in tne Engine/'s parent folder.
+            current_path = editorpath
+            while current_path and current_path != os.path.dirname(current_path):
+                current_path = os.path.dirname(current_path)
+                if os.path.basename(current_path).lower() == 'engine':
+                    engine_parent = os.path.dirname(current_path)
+                    break
+            else:
+                return None
+            
+            import glob
+            uprojectdirs_files = glob.glob(os.path.join(engine_parent, '*.uprojectdirs'))
+            
+            # Read all search paths from all .uprojectdirs files
+
+            search_paths = []
+            for uprojectdirs_file in uprojectdirs_files:
+                try:
+                    with open(uprojectdirs_file, 'r') as f:
+                        for line in f:
+                            line = line.strip()
+                            if line and not line.startswith(';') and not line.startswith('#'):
+                                search_paths.append(line)
+                except:
+                    continue
+            
+            # Search for the project in each path
+
+            for search_path in search_paths:
+                full_search_path = os.path.join(engine_parent, search_path)
+                if os.path.isdir(full_search_path):
+
+                    project_file = os.path.join(full_search_path, f"{project_name}.uproject")
+                    if os.path.exists(project_file):
+                        return os.path.normpath(project_file)
+                    
+                    # Also search in subdirectories named after the project
+                    subdir_project = os.path.join(full_search_path, project_name, f"{project_name}.uproject")
+                    if os.path.exists(subdir_project):
+                        return os.path.normpath(subdir_project)
+            
+            return None
+        except:
+            return None
+
     def p4_settings(self):
         settings = {}
         settings['p4_enabled'] = self.p4_group.isChecked()
@@ -526,7 +578,7 @@ class AddConfigDialog(QtWidgets.QDialog):
         if not self.engine_dir_line_edit.text():
             try:
                 # Try going up the path until you find Engine in a sibling folder. This will work for custom build p4 setups.
-                self.engine_dir_line_edit.setText(self.find_upstream_path_with_name(self.uproject, 'Engine', includeSiblings=True))
+                self.engine_dir_line_edit.setText(str(Path(self.find_upstream_path_with_name(self.uproject, 'Engine', includeSiblings=True))))
             except FileNotFoundError:
                 pass
 
@@ -571,8 +623,7 @@ class AddConfigDialog(QtWidgets.QDialog):
         if not new_dir:
             return
 
-        self.engine_dir = new_dir
-        self.engine_dir = os.path.normpath(self.engine_dir)
+        self.engine_dir = os.path.normpath(str(Path(new_dir)))
         self.engine_dir_line_edit.setText(self.engine_dir)
 
     def update_button_box(self):
