Switchboard: When available, use project's editor target in the UGS build arguments when building the project.

Uses hierarchical detection of the editor target:
* Priority 1: Config/DefaultEngine.ini
* Priotity 2: Source/{PROJECT}Editor.Target.cs
* Priority 3: Source/*Editor.Target.cs

#rnx

diff --git a/Engine/Plugins/VirtualProduction/Switchboard/Source/Switchboard/switchboard/devices/unreal/plugin_unreal.py b/Engine/Plugins/VirtualProduction/Switchboard/Source/Switchboard/switchboard/devices/unreal/plugin_unreal.py
--- a/Engine/Plugins/VirtualProduction/Switchboard/Source/Switchboard/switchboard/devices/unreal/plugin_unreal.py
+++ b/Engine/Plugins/VirtualProduction/Switchboard/Source/Switchboard/switchboard/devices/unreal/plugin_unreal.py
@@ -1939,11 +1939,113 @@ class DeviceUnreal(Device):
         # FIXME?: Not strictly correct, but holds for desktop host platforms
         return self.platform_binary_directory
 
+    def _get_editor_target_name(self) -> Optional[str]:
+        ''' Attempts to find the specific editor target name for the project, mimicking UGS.
+        
+        Uses hierarchical detection of the editor target:
+
+            Priority 1: Config/DefaultEngine.ini
+            Priotity 2: Source/{PROJECT}Editor.Target.cs
+            Priority 3: Source/*Editor.Target.cs
+
+        If not found, we expect the caller to use -Target=Editor in the UGS command line.
+
+        Returns the full editor target name, or None if not found.
+        '''
+
+        uproject_path = CONFIG.UPROJECT_PATH.get_value(self.name)
+
+        if not uproject_path:
+            return None
+            
+        try:
+            uproject_file = pathlib.Path(uproject_path)
+            if not uproject_file.exists():
+                return None
+                
+            project_name = uproject_file.stem
+            
+            # Priority 1: Config/DefaultEngine.ini
+            # Looks for DefaultEditorTarget in its [/Script/BuildSettings.BuildSettings] section
+
+            config_dir = uproject_file.parent / 'Config'
+            engine_ini_path = config_dir / 'DefaultEngine.ini'
+            
+            if engine_ini_path.exists():
+                try:
+                    parser = ugs_utils.IniParser()
+                    with open(engine_ini_path, 'r', encoding='utf-8') as f:
+                        parser.read_file(f)
+                    
+                    default_target = parser.try_get_raw('/Script/BuildSettings.BuildSettings', 'DefaultEditorTarget')
+
+                    if default_target:
+                        LOGGER.debug(f"{self.name}: Found DefaultEditorTarget '{default_target}' in DefaultEngine.ini")
+                        return default_target
+                        
+                except Exception as e:
+                    LOGGER.warning(f"{self.name}: Failed to parse DefaultEngine.ini: {e}")
+            
+            # Priority 2: Source/{PROJECT}Editor.Target.cs
+
+            source_dir = uproject_file.parent / 'Source'
+
+            # No project-specific editor targets in blueprint only projects.
+            if not source_dir.exists():
+                return None
+                
+            file_suffix = '.Target.cs'
+            expected_target_file = source_dir / f"{project_name}Editor{file_suffix}"
+            
+            if expected_target_file.exists():
+                target_name = f"{project_name}Editor"
+                LOGGER.debug(f"{self.name}: Found project-specific editor target: {target_name}")
+                return target_name
+            
+            # Priority 3: Source/*Editor.Target.cs
+
+            editor_targets = list(source_dir.glob(f'*Editor{file_suffix}'))
+            
+            if len(editor_targets) == 1:
+                target_file = editor_targets[0]
+                target_name = target_file.name[:-len(file_suffix)]
+                LOGGER.debug(f"{self.name}: Found single editor target: {target_name}")
+                return target_name
+                    
+            elif len(editor_targets) > 1:
+
+                # Multiple targets found, but no clear winner.
+
+                target_names = [f.name[:-len(file_suffix)] for f in editor_targets]
+
+                LOGGER.warning(
+                    f"{self.name}: Multiple editor targets found but no DefaultEditorTarget configured: {target_names}"
+                )
+                
+        except Exception as e:
+            LOGGER.warning(f"{self.name}: Failed to detect editor target name: {e}")
+            
+        return None
+
     def _build_project(self, puuid_dependency: Optional[uuid.UUID] = None) -> uuid.UUID | None:
-        ubt_args = (
-            f'{self.target_platform} Development '
-            f'-project="{CONFIG.UPROJECT_PATH.get_value(self.name)}" '
-            '-TargetType=Editor -Progress -NoHotReloadFromIDE')
+
+        # Try to get specific project editor target name like UGS does
+        # If found (e.g. MyProjectNameEditor), use it instead of `-TargetType=Editor`
+        # Otherwise fall back to `-TargetType=Editor`
+
+        editor_target = self._get_editor_target_name()
+        
+        if editor_target:
+            ubt_args = (
+                f'{editor_target} {self.target_platform} Development '
+                f'-project="{CONFIG.UPROJECT_PATH.get_value(self.name)}" '
+                '-Progress -NoHotReloadFromIDE')
+        else:
+            ubt_args = (
+                f'{self.target_platform} Development '
+                f'-project="{CONFIG.UPROJECT_PATH.get_value(self.name)}" '
+                '-TargetType=Editor -Progress -NoHotReloadFromIDE')
+
         return self._queue_build('project', ubt_args, puuid_dependency)
 
     def _build_mu_server(self, puuid_dependency: Optional[uuid.UUID] = None) -> uuid.UUID | None:
