Switchboard: Editors are now launched with ensures disabled by default.

[RN] Editors are now launched with ensures disabled by default.

diff --git a/Engine/Plugins/VirtualProduction/Switchboard/Source/Switchboard/switchboard/devices/unreal/plugin_unreal.py b/Engine/Plugins/VirtualProduction/Switchboard/Source/Switchboard/switchboard/devices/unreal/plugin_unreal.py
--- a/Engine/Plugins/VirtualProduction/Switchboard/Source/Switchboard/switchboard/devices/unreal/plugin_unreal.py
+++ b/Engine/Plugins/VirtualProduction/Switchboard/Source/Switchboard/switchboard/devices/unreal/plugin_unreal.py
@@ -812,6 +812,13 @@ class DeviceUnreal(Device):
             ),
             category="General Settings",
         ),
+        'disable_ensures': BoolSetting(
+            attr_name='disable_ensures',
+            nice_name="Disable Ensures",
+            value=True,
+            tool_tip="When checked, disables the handling of ensure errors - which are non-fatal and may cause hitches.",
+            category="UE Settings",
+        ),
     }
 
     if ENABLE_UGS_SUPPORT:
@@ -2439,6 +2446,10 @@ class DeviceUnreal(Device):
                 ' -UDPMESSAGING_TRANSPORT_STATIC='
                 f'"{",".join(static_endpoints)}"')
 
+        # Disable ensures
+        if DeviceUnreal.csettings['disable_ensures'].get_value(self.name):
+            command_line_args += ' -handleensurepercent=0'
+
         record_on_client = 'True' if self.is_recording_device else 'False'
         ini_engine = (
             ' -ini:Engine:'
