Switchboard: Improved level dropdown map search speed.

#rnx

diff --git a/Engine/Plugins/VirtualProduction/Switchboard/Source/Switchboard/switchboard/switchboard_widgets.py b/Engine/Plugins/VirtualProduction/Switchboard/Source/Switchboard/switchboard/switchboard_widgets.py
--- a/Engine/Plugins/VirtualProduction/Switchboard/Source/Switchboard/switchboard/switchboard_widgets.py
+++ b/Engine/Plugins/VirtualProduction/Switchboard/Source/Switchboard/switchboard/switchboard_widgets.py
@@ -2,7 +2,7 @@
 import re
 import socket
 import time
-from typing import Optional
+from typing import Dict, Optional
 
 from PySide6 import QtCore, QtGui, QtWidgets
 
@@ -365,6 +365,9 @@ class SearchableComboBox(QtWidgets.QComboBox):
         self.search_line.textChanged.connect(self.text_changed)
         self.wheel_only_on_focus = True
 
+        # Maps item text to its index for fast lookup during search filtering
+        self._item_index_map: Optional[Dict[str, int]] = None
+
     def showPopup(self):
         super().showPopup()
 
@@ -378,6 +381,9 @@ class SearchableComboBox(QtWidgets.QComboBox):
 
         self.CB_item_text = [self.itemText(i) for i in range(self.count())]
 
+        # Reset index mapping when popup shows in case items changed
+        self._item_index_map = None
+
     def hidePopup(self):
         super().hidePopup()
 
@@ -385,11 +391,42 @@ class SearchableComboBox(QtWidgets.QComboBox):
         self.search_line.clear()
 
     def text_changed(self, text: str):
+
+        # Show all and early exit if text is empty
+        if not text:
+            for i in range(len(self.CB_item_text)):
+                self.view().setRowHidden(i, False)
+            return
+        
+        text_lower = text.lower()
+        
+        # Build index mapping if was reset
+        if self._item_index_map is None:
+            self._item_index_map = {item: i for i, item in enumerate(self.CB_item_text)}
+        
+        # Track visibility changes to minimize UI updates
+        visibility_changes = []
+        
         for item in self.CB_item_text:
-            if item.lower().__contains__(text.lower()):
-                self.view().setRowHidden(self.findText(item), False)
-            else:
-                self.view().setRowHidden(self.findText(item), True)
+            index = self._item_index_map.get(item)
+
+            # This shouldn't happen but let's handle it just in case.
+            if index is None:
+                LOGGER.warning(f"SearchableComboBox: Item '{item}' unexpectedly not found in index map, rebuilding map")
+                self._item_index_map = {item: i for i, item in enumerate(self.CB_item_text)}
+                index = self._item_index_map[item]
+            
+            should_show = text_lower in item.lower()
+            current_hidden = self.view().isRowHidden(index)
+            
+            if should_show and current_hidden:
+                visibility_changes.append((index, False))
+            elif not should_show and not current_hidden:
+                visibility_changes.append((index, True))
+        
+        # Apply visibility changes in batch
+        for index, hidden in visibility_changes:
+            self.view().setRowHidden(index, hidden)
 
     def wheelEvent(self, event) -> None:
 
