Switchboard: Added log verbosity settings to unreal devices.

Only nDisplay had this but this is also useful when debugging the Editor.

#rnx

diff --git a/Engine/Plugins/VirtualProduction/Switchboard/Source/Switchboard/switchboard/devices/unreal/plugin_unreal.py b/Engine/Plugins/VirtualProduction/Switchboard/Source/Switchboard/switchboard/devices/unreal/plugin_unreal.py
--- a/Engine/Plugins/VirtualProduction/Switchboard/Source/Switchboard/switchboard/devices/unreal/plugin_unreal.py
+++ b/Engine/Plugins/VirtualProduction/Switchboard/Source/Switchboard/switchboard/devices/unreal/plugin_unreal.py
@@ -36,7 +36,7 @@ import switchboard.switchboard_widgets as sb_widgets
 from switchboard import message_protocol, switchboard_application, ugs_utils
 from switchboard.config import (
     CONFIG, DEFAULT_MAP_TEXT, ENABLE_UGS_SUPPORT, SETTINGS, BoolSetting, DirectoryPathSetting,
-    EngineSyncMethod, FilePathSetting, FloatSetting, IntSetting, MultiOptionSetting, OptionSetting, Setting,
+    EngineSyncMethod, FilePathSetting, FloatSetting, IntSetting, LoggingSetting, MultiOptionSetting, OptionSetting, Setting,
     StringListSetting, StringSetting, migrate_comma_separated_string_to_list)
 from switchboard.devices.device_base import Device, DeviceStatus, PluginHeaderWidgets
 from switchboard.devices.device_widget_base import DeviceAutoJoinMUServerUI, DeviceWidget
@@ -709,6 +709,23 @@ class DeviceUnreal(Device):
             tool_tip='When checked, retrieves the logs and traces after Unreal Engine terminates.',
             category="UE Settings",
         ),
+        'logging': LoggingSetting(
+            attr_name='logging',
+            nice_name='Logging',
+            value=None,
+            categories=[
+                'LogConcert',
+                'LogImgMedia',
+                'LogLiveLink',
+                'LogMediaIOCore',
+                'LogRemoteControl',
+                'LogRivermaxMedia',
+                'LogSharedMemoryMedia',
+                'LogUdpMessaging',
+            ],
+            tool_tip='Logging categories and verbosity levels',
+            category="UE Settings",
+        ),
         'livelink_preset': LiveLinkPresetSetting(
             attr_name='livelink_preset',
             nice_name='LiveLink Preset',
@@ -2458,6 +2475,12 @@ class DeviceUnreal(Device):
         )
         command_line_args += ini_engine
 
+        # Logging
+        logging_args = DeviceUnreal.csettings['logging'].get_command_line_arg(
+            override_device_name=self.name)
+        if logging_args:
+            command_line_args += f' {logging_args}'
+
         return (
             f'"{CONFIG.UPROJECT_PATH.get_value(self.name)}" {map_name} '
             f'{command_line_args}')
