Switchboard: Add a timer to keep track of how long devices has been running.

You can see this in the tooltip of the status icon. It will update it real time, no need to refresh the tooltip.

#rnx

diff --git a/Engine/Plugins/VirtualProduction/Switchboard/Source/Switchboard/switchboard/device_list_widget.py b/Engine/Plugins/VirtualProduction/Switchboard/Source/Switchboard/switchboard/device_list_widget.py
--- a/Engine/Plugins/VirtualProduction/Switchboard/Source/Switchboard/switchboard/device_list_widget.py
+++ b/Engine/Plugins/VirtualProduction/Switchboard/Source/Switchboard/switchboard/device_list_widget.py
@@ -206,7 +206,7 @@ class DeviceListWidget(QtWidgets.QListWidget):
 
         self.signal_register_device_widget.emit(device.widget)
         # force the widget to update, to make sure status is correct
-        device.widget.update_status(device.status, device.status)
+        device.widget.update_status(device.status, device.status, device)
 
         self.connect_to_device_autojoin(device, category)
 
diff --git a/Engine/Plugins/VirtualProduction/Switchboard/Source/Switchboard/switchboard/devices/device_base.py b/Engine/Plugins/VirtualProduction/Switchboard/Source/Switchboard/switchboard/devices/device_base.py
--- a/Engine/Plugins/VirtualProduction/Switchboard/Source/Switchboard/switchboard/devices/device_base.py
+++ b/Engine/Plugins/VirtualProduction/Switchboard/Source/Switchboard/switchboard/devices/device_base.py
@@ -6,6 +6,7 @@ import os
 import random
 import socket
 import threading
+import time
 import weakref
 
 from enum import IntEnum, IntFlag, auto, unique
@@ -120,6 +121,9 @@ class Device(QtCore.QObject):
 
         self.device_recording = None
 
+        # Runtime tracking
+        self._runtime_start_time: Optional[float] = None
+
         self.widget: DeviceWidget = None  # Constructed during init()
 
     def init(self, widget_class, icons):
@@ -252,6 +256,13 @@ class Device(QtCore.QObject):
     def status(self, value: DeviceStatus):
         previous_status = self._status
         self._status = value
+        
+        # Track runtime when device starts/stops running
+        if value == DeviceStatus.OPEN and previous_status != DeviceStatus.OPEN:
+            self._runtime_start_time = time.time()
+        elif previous_status == DeviceStatus.OPEN and value != DeviceStatus.OPEN:
+            self._runtime_start_time = None
+            
         self.device_qt_handler.signal_device_status_changed.emit(
             self, previous_status)
 
@@ -260,6 +271,34 @@ class Device(QtCore.QObject):
         return self.status in {
             DeviceStatus.DISCONNECTED, DeviceStatus.CONNECTING}
 
+    def get_runtime_string(self) -> Optional[str]:
+        ''' Returns a string representing how long the device has been running.
+        If the device is not actually running, returns None.
+
+        Shows hours when >= 1 hour, and days when >= 24 hours.
+
+        e.g.
+                  "01:02" (1 min, 2 seconds)
+               "01:02:03" (1 hour, 2 minutes, 3 seconds)
+            "1d 02:03:04" (1 day, 2 hours, 3 minutes, 4 seconds)
+        '''
+
+        if self.status != DeviceStatus.OPEN or self._runtime_start_time is None:
+            return None
+            
+        runtime_seconds = int(time.time() - self._runtime_start_time)
+        days = runtime_seconds // 86400
+        hours = (runtime_seconds % 86400) // 3600
+        minutes = (runtime_seconds % 3600) // 60
+        seconds = runtime_seconds % 60
+        
+        if days > 0:
+            return f"{days}d {hours:02d}:{minutes:02d}:{seconds:02d}"
+        elif hours > 0:
+            return f"{hours:02d}:{minutes:02d}:{seconds:02d}"
+        else:
+            return f"{minutes:02d}:{seconds:02d}"
+
     @property
     def project_changelist(self):
         return self._project_changelist
diff --git a/Engine/Plugins/VirtualProduction/Switchboard/Source/Switchboard/switchboard/devices/device_widget_base.py b/Engine/Plugins/VirtualProduction/Switchboard/Source/Switchboard/switchboard/devices/device_widget_base.py
--- a/Engine/Plugins/VirtualProduction/Switchboard/Source/Switchboard/switchboard/devices/device_widget_base.py
+++ b/Engine/Plugins/VirtualProduction/Switchboard/Source/Switchboard/switchboard/devices/device_widget_base.py
@@ -5,10 +5,72 @@ from PySide6 import QtWidgets, QtGui, QtCore
 
 from switchboard.config import CONFIG
 from switchboard.switchboard_widgets import FramelessQLineEdit
-from switchboard.devices.device_base import DeviceStatus
+from switchboard.devices.device_base import DeviceStatus, Device
 import switchboard.switchboard_widgets as sb_widgets
 
 
+class DeviceRunningTimeTooltip(QtWidgets.QWidget):
+    ''' Custom tooltip widget that can constantly update its display of the device's running time.
+    
+    It should only apply to devices currently with OPEN status.
+    '''
+    
+    def __init__(self, parent=None):
+        super().__init__(parent, QtCore.Qt.ToolTip | QtCore.Qt.FramelessWindowHint)
+
+        self.setAttribute(QtCore.Qt.WA_ShowWithoutActivating)
+        
+        self.label = QtWidgets.QLabel()
+        self.label.setWordWrap(True)
+        layout = QtWidgets.QVBoxLayout()
+        layout.setContentsMargins(0, 0, 0, 0)
+        layout.addWidget(self.label)
+        self.setLayout(layout)
+        
+        # Timer to update the display every second
+
+        self.update_timer = QtCore.QTimer()
+        self.update_timer.timeout.connect(self.update_content)
+        self.update_timer.setInterval(1000)
+        
+        self.device: Optional[Device] = None
+        
+    def show_for_device(self, device, position):
+        '''Show the tooltip for the given device at the given position.
+        '''
+        self.device = device
+        self.update_content()
+        self.move(position)
+        self.show()
+        self.update_timer.start()
+        
+    def hide_tooltip(self):
+        ''' Hide the tooltip and stop updates.
+        '''
+        self.update_timer.stop()
+        self.hide()
+        self.device = None
+        
+    def update_content(self):
+        ''' Update the tooltip content the device's current running time.
+        '''
+
+        if not self.device or self.device.status != DeviceStatus.OPEN:
+            self.hide_tooltip()
+            return
+
+        # Read the device's running time as a string to display in the tooltip            
+        runtime_str = self.device.get_runtime_string()
+
+        if runtime_str:
+            tooltip_text = f"Device has been running for {runtime_str}"
+        else:
+            tooltip_text = "Device has been started"
+        
+        self.label.setText(tooltip_text)
+        self.adjustSize()
+
+
 class DeviceWidgetItem(QtWidgets.QWidget):
     """
     Custom class to get QSS working correctly to achieve a look.
@@ -96,6 +158,11 @@ class DeviceWidget(QtWidgets.QWidget):
         pixmap = QtGui.QPixmap(":/icons/images/status_blank_disabled.png")
         self.status_icon.setPixmap(pixmap)
         self.status_icon.resize(pixmap.width(), pixmap.height())
+        
+        # Enable mouse tracking for hover events
+        self.status_icon.setMouseTracking(True)
+        self.status_icon.enterEvent = self._on_status_icon_enter
+        self.status_icon.leaveEvent = self._on_status_icon_leave
 
         # Device icon
         self.device_icon = QtWidgets.QLabel()
@@ -160,6 +227,16 @@ class DeviceWidget(QtWidgets.QWidget):
         # Store the control buttons by name ("connect", "open", etc.)
         self.control_buttons: Dict[str, sb_widgets.ControlQPushButton] = {}
 
+        # Runtime tracking
+        self._current_device = None
+        self._runtime_tooltip = DeviceRunningTimeTooltip(self)
+        
+        # Timer for tooltip delay
+        self._tooltip_delay_timer = QtCore.QTimer()
+        self._tooltip_delay_timer.setSingleShot(True)
+        self._tooltip_delay_timer.timeout.connect(self._show_delayed_tooltip)
+        self._pending_tooltip_pos = None
+
         self._add_control_buttons()
 
     def add_widget_to_layout(self, widget):
@@ -226,7 +303,54 @@ class DeviceWidget(QtWidgets.QWidget):
     def _add_control_buttons(self):
         pass
 
-    def update_status(self, status, previous_status):
+    def _on_status_icon_enter(self, event):
+        ''' Show custom tooltip when mouse enters status icon.
+        '''
+
+        if self._current_device and self._current_device.status == DeviceStatus.OPEN:
+            # Position tooltip similar to the regular tooltip
+            self._pending_tooltip_pos = QtGui.QCursor.pos() + QtCore.QPoint(10, 10)
+
+            # Start delay timer. Without the timer the tooltip appears too fast and before you center the mouse
+            self._tooltip_delay_timer.start(700)
+        else:
+            # For non-OPEN devices, show the regular tooltip
+            self._show_regular_tooltip()
+
+    def _on_status_icon_leave(self, event):
+        ''' Hide custom tooltip when mouse leaves status icon.
+        '''
+        # Cancel pending tooltip
+        self._tooltip_delay_timer.stop()
+        self._runtime_tooltip.hide_tooltip()
+        
+    def _show_delayed_tooltip(self):
+        ''' Show the custom tooltip after delay.
+        '''
+        if self._current_device and self._current_device.status == DeviceStatus.OPEN and self._pending_tooltip_pos:
+            self._runtime_tooltip.show_for_device(self._current_device, self._pending_tooltip_pos)
+
+    def _show_regular_tooltip(self):
+        ''' Show the appropriate regular tooltip based on device status.
+        '''
+
+        if not self._current_device:
+            return
+            
+        status = self._current_device.status
+        if status >= DeviceStatus.READY:
+            self.status_icon.setToolTip("Ready to start recording")
+        elif status == DeviceStatus.DISCONNECTED:
+            self.status_icon.setToolTip("Disconnected")
+        elif status == DeviceStatus.CONNECTING:
+            self.status_icon.setToolTip("Connecting...")
+        else:
+            self.status_icon.setToolTip("Connected")
+
+    def update_status(self, status, previous_status, device=None):
+        # Store device reference for runtime updates
+        self._current_device = device
+        
         # Status Icon
         if status >= DeviceStatus.READY:
             self.status_icon.setPixmap(QtGui.QPixmap(":/icons/images/status_green.png"))
@@ -242,7 +366,8 @@ class DeviceWidget(QtWidgets.QWidget):
         elif status == DeviceStatus.OPEN:
             pixmap = QtGui.QPixmap(":/icons/images/status_orange.png")
             self.status_icon.setPixmap(pixmap)
-            self.status_icon.setToolTip("Device has been started")
+            # Clear regular tooltip for running devices (custom tooltip will handle this)
+            self.status_icon.setToolTip("")
         else:
             self.status_icon.setPixmap(QtGui.QPixmap(":/icons/images/status_cyan.png"))
             self.status_icon.setToolTip("Connected")
diff --git a/Engine/Plugins/VirtualProduction/Switchboard/Source/Switchboard/switchboard/devices/kipro/plugin_kipro.py b/Engine/Plugins/VirtualProduction/Switchboard/Source/Switchboard/switchboard/devices/kipro/plugin_kipro.py
--- a/Engine/Plugins/VirtualProduction/Switchboard/Source/Switchboard/switchboard/devices/kipro/plugin_kipro.py
+++ b/Engine/Plugins/VirtualProduction/Switchboard/Source/Switchboard/switchboard/devices/kipro/plugin_kipro.py
@@ -290,8 +290,8 @@ class DeviceWidgetKiPro(DeviceWidget):
 
         self.connect_button.clicked.connect(self.connect_button_clicked)
 
-    def update_status(self, status, previous_status):
-        super().update_status(status, previous_status)
+    def update_status(self, status, previous_status, device=None):
+        super().update_status(status, previous_status, device)
 
         if status in {DeviceStatus.DISCONNECTED, DeviceStatus.CONNECTING}:
             self.play_button.setDisabled(True)
diff --git a/Engine/Plugins/VirtualProduction/Switchboard/Source/Switchboard/switchboard/devices/unreal/plugin_unreal.py b/Engine/Plugins/VirtualProduction/Switchboard/Source/Switchboard/switchboard/devices/unreal/plugin_unreal.py
--- a/Engine/Plugins/VirtualProduction/Switchboard/Source/Switchboard/switchboard/devices/unreal/plugin_unreal.py
+++ b/Engine/Plugins/VirtualProduction/Switchboard/Source/Switchboard/switchboard/devices/unreal/plugin_unreal.py
@@ -3947,8 +3947,8 @@ class DeviceWidgetUnreal(DeviceWidget):
         self.sync_button.setDisabled(False)
         self.build_button.setDisabled(self.exclude_from_build)
 
-    def update_status(self, status, previous_status):
-        super().update_status(status, previous_status)
+    def update_status(self, status, previous_status, device=None):
+        super().update_status(status, previous_status, device)
 
         if status <= DeviceStatus.CONNECTING:
             self._update_disconnected_ui()
diff --git a/Engine/Plugins/VirtualProduction/Switchboard/Source/Switchboard/switchboard/switchboard_dialog.py b/Engine/Plugins/VirtualProduction/Switchboard/Source/Switchboard/switchboard/switchboard_dialog.py
--- a/Engine/Plugins/VirtualProduction/Switchboard/Source/Switchboard/switchboard/switchboard_dialog.py
+++ b/Engine/Plugins/VirtualProduction/Switchboard/Source/Switchboard/switchboard/switchboard_dialog.py
@@ -1984,7 +1984,7 @@ class SwitchboardDialog(QtCore.QObject):
     @QtCore.Slot(object)
     def device_status_changed(self, device, previous_status):
         # Update the device widget
-        device.widget.update_status(device.status, previous_status)
+        device.widget.update_status(device.status, previous_status, device)
 
         devices = self.device_manager.devices_of_type(device.device_type)
         self.device_list_widget.update_category_status(device.category_name, devices)
