SB: Fix a couple of exceptions due to acessing deleted Qt c++ objects.

One was due to not disconnecting a global CONFIG signal from a widget.

The other was due to ListSetting not cleaning up element_layouts.

#rnx

diff --git a/Engine/Plugins/VirtualProduction/Switchboard/Source/Switchboard/switchboard/config.py b/Engine/Plugins/VirtualProduction/Switchboard/Source/Switchboard/switchboard/config.py
--- a/Engine/Plugins/VirtualProduction/Switchboard/Source/Switchboard/switchboard/config.py
+++ b/Engine/Plugins/VirtualProduction/Switchboard/Source/Switchboard/switchboard/config.py
@@ -1640,6 +1640,14 @@ class ListSetting(Setting):
 
         self._update_array_count_label(override_device_name)
 
+    def _on_widget_destroyed(self, destroyed_object=None, override_device_name: str | None = None):
+        super()._on_widget_destroyed(destroyed_object, override_device_name)
+        
+        # Clean up the layout references to prevent exceptions related to accessing deleted Qt C++ objects.
+        if override_device_name is not None:
+            self.element_layouts.pop(override_device_name, None)
+            self.array_count_labels.pop(override_device_name, None)
+
 
 class StringListSetting(ListSetting):
     '''
diff --git a/Engine/Plugins/VirtualProduction/Switchboard/Source/Switchboard/switchboard/devices/device_widget_base.py b/Engine/Plugins/VirtualProduction/Switchboard/Source/Switchboard/switchboard/devices/device_widget_base.py
--- a/Engine/Plugins/VirtualProduction/Switchboard/Source/Switchboard/switchboard/devices/device_widget_base.py
+++ b/Engine/Plugins/VirtualProduction/Switchboard/Source/Switchboard/switchboard/devices/device_widget_base.py
@@ -94,6 +94,13 @@ class DeviceAutoJoinMUServerUI(QtCore.QObject):
     def __init__(self, name, parent = None):
         super().__init__(parent)
         self.name = name
+        self._button = None
+
+    def __del__(self):
+        try:
+            CONFIG.MUSERVER_AUTO_JOIN.signal_setting_changed.disconnect(self.disable_enable_based_on_global)
+        except (RuntimeError, TypeError):
+            pass
 
     def is_autojoin_enabled(self):
         return self._button.isChecked()
