Switchboard: Suppress QUIC spam log from the log window.

Also fixed a previous qt c++ object reference cleanup issue.

#rnx

diff --git a/Engine/Plugins/VirtualProduction/Switchboard/Source/Switchboard/switchboard/config.py b/Engine/Plugins/VirtualProduction/Switchboard/Source/Switchboard/switchboard/config.py
--- a/Engine/Plugins/VirtualProduction/Switchboard/Source/Switchboard/switchboard/config.py
+++ b/Engine/Plugins/VirtualProduction/Switchboard/Source/Switchboard/switchboard/config.py
@@ -739,6 +739,13 @@ class Setting(QtCore.QObject):
             ''' Captures override_device_name but not self to avoid memory leaks due to strong reference '''
             self_instance = self_weak_ref()
             if self_instance is not None:
+
+                # Disconnect signal_setting_changed to prevent accumulation
+                try:
+                    self_instance.signal_setting_changed.disconnect(pre_on_setting_changed_weak)
+                except (RuntimeError, TypeError):
+                    pass
+
                 self_instance._on_widget_destroyed(override_device_name=override_device_name)
 
         top_level_widget.destroyed.connect(on_widget_destroyed_weak)
@@ -1644,9 +1651,8 @@ class ListSetting(Setting):
         super()._on_widget_destroyed(destroyed_object, override_device_name)
         
         # Clean up the layout references to prevent exceptions related to accessing deleted Qt C++ objects.
-        if override_device_name is not None:
-            self.element_layouts.pop(override_device_name, None)
-            self.array_count_labels.pop(override_device_name, None)
+        self.element_layouts.pop(override_device_name, None)
+        self.array_count_labels.pop(override_device_name, None)
 
 
 class StringListSetting(ListSetting):
diff --git a/Engine/Plugins/VirtualProduction/Switchboard/Source/Switchboard/switchboard/switchboard_logging.py b/Engine/Plugins/VirtualProduction/Switchboard/Source/Switchboard/switchboard/switchboard_logging.py
--- a/Engine/Plugins/VirtualProduction/Switchboard/Source/Switchboard/switchboard/switchboard_logging.py
+++ b/Engine/Plugins/VirtualProduction/Switchboard/Source/Switchboard/switchboard/switchboard_logging.py
@@ -31,6 +31,11 @@ class QtHandler(logging.Handler):
         self.formatter = logging.Formatter()
 
     def emit(self, record: logging.LogRecord):
+
+        # Filter spammy debug messages from UI (but they can still be found in the log file)
+        if record.name in ('quic', 'asyncio') and record.levelno < logging.WARNING:
+            return
+            
         html_record = self.format(record)
         ConsoleStream.stdout().write(html_record)
         ConsoleStream.stderr().write(html_record)
@@ -153,8 +158,11 @@ LOGGER = cast(SBLogger, logging.getLogger(DEFAULT_LOGGER_NAME))
 ROOT_LOGGER = logging.getLogger()
 
 QT_HANDLER = QtHandler()
-ROOT_LOGGER.addHandler(QT_HANDLER)
+ROOT_LOGGER.addHandler(QT_HANDLER) 
 
 FILE_HANDLER = logging.FileHandler(SBLogger.make_log_path())
 FILE_HANDLER.setFormatter(logging.Formatter(fmt=LOGGING_FMT))
 ROOT_LOGGER.addHandler(FILE_HANDLER)
+
+# Note: QUIC and asyncio debug message filtering is handled in QtHandler.emit()
+# This allows debug messages to still reach the file log while suppressing UI spam
