Switchboard: Fix for DevicenDisplay.close_all_devices transitioning disconnected devices to closing state.

#rnx

diff --git a/Engine/Plugins/VirtualProduction/Switchboard/Source/Switchboard/switchboard/devices/ndisplay/plugin_ndisplay.py b/Engine/Plugins/VirtualProduction/Switchboard/Source/Switchboard/switchboard/devices/ndisplay/plugin_ndisplay.py
--- a/Engine/Plugins/VirtualProduction/Switchboard/Source/Switchboard/switchboard/devices/ndisplay/plugin_ndisplay.py
+++ b/Engine/Plugins/VirtualProduction/Switchboard/Source/Switchboard/switchboard/devices/ndisplay/plugin_ndisplay.py
@@ -1,5 +1,7 @@
 # Copyright Epic Games, Inc. All Rights Reserved.
 
+from __future__ import annotations
+
 import json
 import os
 import socket
@@ -2019,7 +2021,7 @@ class DevicenDisplay(DeviceUnreal):
     _devices_pending_shutdown = []
 
     @classmethod
-    def close_all_devices(cls, devices):
+    def close_all_devices(cls, devices: list[DevicenDisplay]):
         '''
         Override to use nDisplay's cluster-aware shutdown with timeout fallback.
         First tries cluster soft kill, then falls back to individual closes after timeout.
@@ -2027,38 +2029,39 @@ class DevicenDisplay(DeviceUnreal):
         '''
         if not devices:
             return
-            
+
         full_timeout = DeviceUnreal.csettings['soft_kill_timeout'].get_value()
         cluster_timeout = full_timeout / 2.0  # Half timeout for cluster phase
-        
+
         try:
             # Issue cluster command, which is preferred because it issues the quit command at the same time on all
             # nodes which helps normal shutdown without thread barrier related deadlocks.
             cls.soft_kill_cluster(devices)
-            
+
             # Since the cluster command may not succeed, after a timeout we issue the regular kill command via 
             # SwitchboardListener on the nodes that are still running.
 
             cls._devices_pending_shutdown = devices.copy()
-            
+
             # Cancel any existing timeout timer
             if cls._cluster_shutdown_timer:
                 cls._cluster_shutdown_timer.stop()
                 cls._cluster_shutdown_timer = None
-            
+
             # Create new timer for the fallback kill command - use half the timeout
             cls._cluster_shutdown_timer = QtCore.QTimer()
             cls._cluster_shutdown_timer.timeout.connect(cls._on_cluster_shutdown_timeout)
             cls._cluster_shutdown_timer.setSingleShot(True)
             cls._cluster_shutdown_timer.start(int(cluster_timeout * 1000))  # Convert to ms
-            
+
             # Set devices to CLOSING state (disables buttons and shows proper state)
             for device in devices:
                 try:
-                    device.status = DeviceStatus.CLOSING
+                    if device.status in (DeviceStatus.OPEN, DeviceStatus.READY):
+                        device.status = DeviceStatus.CLOSING
                 except Exception:
                     pass
-                    
+
         except Exception:
             LOGGER.warning("Cluster event for shutdown could not be issued, falling back to immediate device close")
 
@@ -2070,17 +2073,17 @@ class DevicenDisplay(DeviceUnreal):
         '''Called after timeout - check if devices are still running and force close if needed'''
         if not cls._devices_pending_shutdown:
             return
-                    
+
         # Check which devices are still running and force close them
         devices_still_running = []
         for device in cls._devices_pending_shutdown:
             # Check if device is still in CLOSING state (meaning cluster shutdown didn't complete)
             if device.status == DeviceStatus.CLOSING:
                 devices_still_running.append(device)
-        
-        if devices_still_running:            
+
+        if devices_still_running:
             cls._close_nodes_individually(devices_still_running)
-        
+
         cls._devices_pending_shutdown.clear()
         cls._cluster_shutdown_timer = None
 
@@ -2089,11 +2092,12 @@ class DevicenDisplay(DeviceUnreal):
         '''Force close individual devices using the standard close mechanism with remaining timeout'''
         full_timeout = DeviceUnreal.csettings['soft_kill_timeout'].get_value()
         remaining_timeout = full_timeout / 2.0  # Use remaining half of timeout
-        
+
         for device in devices:
             try:
                 # Call device.close() directly with the remaining timeout
-                device.close(soft_kill_timeout=remaining_timeout)
+                if device.status in (DeviceStatus.OPEN, DeviceStatus.READY, DeviceStatus.CLOSING):
+                    device.close(soft_kill_timeout=remaining_timeout)
             except Exception:
                 pass
 
