Switchboard: Hide "sync" button, disable "sync all" and "sync+build all" buttons when CL dropdown refresh is pending.

Also catch earlier if a single device sync is attempted with an invalid dropdown selection.

#rnx

diff --git a/Engine/Plugins/VirtualProduction/Switchboard/Source/Switchboard/switchboard/devices/unreal/plugin_unreal.py b/Engine/Plugins/VirtualProduction/Switchboard/Source/Switchboard/switchboard/devices/unreal/plugin_unreal.py
--- a/Engine/Plugins/VirtualProduction/Switchboard/Source/Switchboard/switchboard/devices/unreal/plugin_unreal.py
+++ b/Engine/Plugins/VirtualProduction/Switchboard/Source/Switchboard/switchboard/devices/unreal/plugin_unreal.py
@@ -1718,6 +1718,20 @@ class DeviceUnreal(Device):
                 "There is nothing to sync!")
             return None
 
+        if isinstance(engine_cl, str):
+            try:
+                engine_cl = int(engine_cl)
+            except ValueError:
+                LOGGER.error(f'Invalid engine changelist specified for sync: {engine_cl}')
+                return None
+
+        if isinstance(project_cl, str):
+            try:
+                project_cl = int(project_cl)
+            except ValueError:
+                LOGGER.error(f'Invalid project changelist specified for sync: {project_cl}')
+                return None
+
         for device in self.devices_sharing_workspace():
             if device.status >= DeviceStatus.OPEN:
                 LOGGER.warning(f'{self.name}: Unable to sync, because another '
@@ -4019,6 +4033,11 @@ class DeviceWidgetUnreal(DeviceWidget):
         self.sync_button.show()
         self.build_button.show()
 
+        try:
+            required_cl_int = int(required_cl)
+        except ValueError:
+            self.sync_button.hide()
+
         is_synched = required_cl is None or required_cl == current_device_cl
         self._set_project_changelist_is_synched(is_synched)
 
@@ -4070,6 +4089,11 @@ class DeviceWidgetUnreal(DeviceWidget):
         self.sync_button.show()
         self.build_button.show()
 
+        try:
+            required_cl_int = int(required_cl)
+        except ValueError:
+            self.sync_button.hide()
+
         is_synched = required_cl is None or required_cl == synced_cl
         self._set_engine_changelist_is_synched(is_synched)
         self.update_build_info(device)
diff --git a/Engine/Plugins/VirtualProduction/Switchboard/Source/Switchboard/switchboard/switchboard_dialog.py b/Engine/Plugins/VirtualProduction/Switchboard/Source/Switchboard/switchboard/switchboard_dialog.py
--- a/Engine/Plugins/VirtualProduction/Switchboard/Source/Switchboard/switchboard/switchboard_dialog.py
+++ b/Engine/Plugins/VirtualProduction/Switchboard/Source/Switchboard/switchboard/switchboard_dialog.py
@@ -2203,8 +2203,13 @@ class SwitchboardDialog(QtCore.QObject):
             ):
                 paths.extend([f'//{client}{x}' for x in addpaths])
 
+        # Indeterminate state, resolved on runnable completion
         self.window.project_cl_combo_box.clear()
-        self.window.project_cl_combo_box.addItem('...')
+        self._set_project_changelist('...')
+        self.window.refresh_project_cl_button.setEnabled(False)
+        self.window.sync_all_button.setEnabled(False)
+        self.window.sync_and_build_all_button.setEnabled(False)
+
         self._project_cl_runnable = SwitchboardDialog.ChangelistRunnable(paths, client=client)
         self._project_cl_runnable.signals.result.connect(self._on_project_changelists_refreshed)
         QtCore.QThreadPool.globalInstance().start(self._project_cl_runnable)
@@ -2216,6 +2221,10 @@ class SwitchboardDialog(QtCore.QObject):
             self._project_cl_runnable = None
 
         self.window.project_cl_combo_box.clear()
+        self.window.refresh_project_cl_button.setEnabled(True)
+        no_pending_refresh = (self._project_cl_runnable is None) and (self._engine_cl_runnable is None)
+        self.window.sync_all_button.setEnabled(no_pending_refresh)
+        self.window.sync_and_build_all_button.setEnabled(no_pending_refresh)
 
         for (idx, (cl, desc)) in zip(count(), cl_desc_pairs):
             self.window.project_cl_combo_box.addItem(str(cl))
@@ -2244,8 +2253,12 @@ class SwitchboardDialog(QtCore.QObject):
 
             engine_p4_path = CONFIG.P4_ENGINE_PATH.get_value()
             if engine_p4_path:
+                # Indeterminate state, resolved on runnable completion
                 self.window.engine_cl_combo_box.clear()
-                self.window.engine_cl_combo_box.addItem('...')
+                self._set_engine_changelist('...')
+                self.window.refresh_engine_cl_button.setEnabled(False)
+                self.window.sync_all_button.setEnabled(False)
+                self.window.sync_and_build_all_button.setEnabled(False)
 
                 client = CONFIG.SOURCE_CONTROL_WORKSPACE.get_value()
                 self._engine_cl_runnable = SwitchboardDialog.ChangelistRunnable([engine_p4_path+'/...'], client=client)
@@ -2271,6 +2284,10 @@ class SwitchboardDialog(QtCore.QObject):
             self._engine_cl_runnable = None
 
         self.window.engine_cl_combo_box.clear()
+        self.window.refresh_engine_cl_button.setEnabled(True)
+        no_pending_refresh = (self._project_cl_runnable is None) and (self._engine_cl_runnable is None)
+        self.window.sync_all_button.setEnabled(no_pending_refresh)
+        self.window.sync_and_build_all_button.setEnabled(no_pending_refresh)
 
         for (idx, (cl, desc)) in zip(count(), cl_desc_pairs):
             self.window.engine_cl_combo_box.addItem(str(cl))
