Switchboard: Add support for EUnrealEngineObjectUE5Version::IMPORT_TYPE_HIERARCHIES to UassetParser.

(Introduced in CL 45907013.)

#rnx

diff --git a/Engine/Plugins/VirtualProduction/Switchboard/Source/Switchboard/switchboard/devices/unreal/uassetparser.py b/Engine/Plugins/VirtualProduction/Switchboard/Source/Switchboard/switchboard/devices/unreal/uassetparser.py
--- a/Engine/Plugins/VirtualProduction/Switchboard/Source/Switchboard/switchboard/devices/unreal/uassetparser.py
+++ b/Engine/Plugins/VirtualProduction/Switchboard/Source/Switchboard/switchboard/devices/unreal/uassetparser.py
@@ -24,6 +24,7 @@ class EUnrealEngineObjectUE5Version(IntEnum):
     METADATA_SERIALIZATION_OFFSET = 1014
     VERSE_CELLS = 1015
     PACKAGE_SAVED_HASH = 1016
+    IMPORT_TYPE_HIERARCHIES = 1018
 
 
 class UassetSummary(object):
@@ -314,7 +315,7 @@ class UassetParser(object):
 
         s.FileVersionLicenseeUE4 = self.readUInt32()
 
-        KNOWN_SUPPORTED_UE5VER = 1017
+        KNOWN_SUPPORTED_UE5VER = 1018
         if s.FileVersionUE5 > KNOWN_SUPPORTED_UE5VER:
             LOGGER.warning(f"{self.f.name}: ObjectUE5Version {s.FileVersionUE5} too new; "
                            f"newest known supported version {KNOWN_SUPPORTED_UE5VER}")
@@ -358,6 +359,10 @@ class UassetParser(object):
         s.SearchableNamesOffset, s.ThumbnailTableOffset \
             = struct.unpack('<5i', self.f.read(5*4))
 
+        if s.FileVersionUE5 >= EUnrealEngineObjectUE5Version.IMPORT_TYPE_HIERARCHIES:
+            s.ImportTypeHierarchiesCount, s.ImportTypeHierarchiesOffset \
+                = struct.unpack('<2i', self.f.read(2*4))
+
         if s.FileVersionUE5 < EUnrealEngineObjectUE5Version.PACKAGE_SAVED_HASH:
             s.Guid = struct.unpack('<16s', self.f.read(16))
 
