Switchboard: Fix for None `required_cl`.

diff --git a/Engine/Plugins/VirtualProduction/Switchboard/Source/Switchboard/switchboard/devices/unreal/plugin_unreal.py b/Engine/Plugins/VirtualProduction/Switchboard/Source/Switchboard/switchboard/devices/unreal/plugin_unreal.py
--- a/Engine/Plugins/VirtualProduction/Switchboard/Source/Switchboard/switchboard/devices/unreal/plugin_unreal.py
+++ b/Engine/Plugins/VirtualProduction/Switchboard/Source/Switchboard/switchboard/devices/unreal/plugin_unreal.py
@@ -4025,7 +4025,7 @@ class DeviceWidgetUnreal(DeviceWidget):
         else:
             self.open_button.setToolTip(f'Start {self.app_name}')
 
-    def update_project_changelist(self, required_cl: str, current_device_cl: str):
+    def update_project_changelist(self, required_cl: str | None, current_device_cl: str):
         self.project_changelist_label.setText(f'P: {current_device_cl}')
         self.project_changelist_label.setToolTip('Project CL')
 
@@ -4035,7 +4035,7 @@ class DeviceWidgetUnreal(DeviceWidget):
 
         try:
             required_cl_int = int(required_cl)
-        except ValueError:
+        except (TypeError, ValueError):
             self.sync_button.hide()
 
         is_synched = required_cl is None or required_cl == current_device_cl
@@ -4074,7 +4074,7 @@ class DeviceWidgetUnreal(DeviceWidget):
 
     def update_engine_changelist(
         self,
-        required_cl: str,
+        required_cl: str | None,
         device: DeviceUnreal,
     ):
         if not CONFIG.ENGINE_SYNC_METHOD.get_value() == EngineSyncMethod.Build_Engine.value:
@@ -4091,7 +4091,7 @@ class DeviceWidgetUnreal(DeviceWidget):
 
         try:
             required_cl_int = int(required_cl)
-        except ValueError:
+        except (TypeError, ValueError):
             self.sync_button.hide()
 
         is_synched = required_cl is None or required_cl == synced_cl
