Switchboard: AddDeviceDialog - Don't discard keystrokes that would lead to duplicate device names.

diff --git a/Engine/Plugins/VirtualProduction/Switchboard/Source/Switchboard/switchboard/devices/device_manager.py b/Engine/Plugins/VirtualProduction/Switchboard/Source/Switchboard/switchboard/devices/device_manager.py
--- a/Engine/Plugins/VirtualProduction/Switchboard/Source/Switchboard/switchboard/devices/device_manager.py
+++ b/Engine/Plugins/VirtualProduction/Switchboard/Source/Switchboard/switchboard/devices/device_manager.py
@@ -218,4 +218,4 @@ class DeviceNameValidator(QtGui.QValidator):
     def validate(self, input, pos):
         if self.device_manager.is_name_unique(input):
             return QtGui.QValidator.Acceptable
-        return QtGui.QValidator.Invalid
+        return QtGui.QValidator.Intermediate
diff --git a/Engine/Plugins/VirtualProduction/Switchboard/Source/Switchboard/switchboard/devices/device_widget_base.py b/Engine/Plugins/VirtualProduction/Switchboard/Source/Switchboard/switchboard/devices/device_widget_base.py
--- a/Engine/Plugins/VirtualProduction/Switchboard/Source/Switchboard/switchboard/devices/device_widget_base.py
+++ b/Engine/Plugins/VirtualProduction/Switchboard/Source/Switchboard/switchboard/devices/device_widget_base.py
@@ -448,9 +448,11 @@ class AddDeviceDialog(QtWidgets.QDialog):
         self.setWindowTitle(f"Add {self.device_type} Device")
 
         self.name_field = QtWidgets.QLineEdit(self)
+        self.name_field.textChanged.connect(lambda: self.refresh_validity())
 
         self.address_field = QtWidgets.QLineEdit(self)
         self.address_field.setValidator(DeviceWidget.hostname_validator)
+        self.address_field.textChanged.connect(lambda: self.refresh_validity())
 
         self.form_layout = QtWidgets.QFormLayout()
         self.form_layout.addRow("Name", self.name_field)
@@ -459,17 +461,30 @@ class AddDeviceDialog(QtWidgets.QDialog):
         layout = QtWidgets.QVBoxLayout()
         layout.insertLayout(0, self.form_layout)
 
-        button_box = QtWidgets.QDialogButtonBox(QtWidgets.QDialogButtonBox.Ok | QtWidgets.QDialogButtonBox.Cancel)
-        button_box.accepted.connect(lambda: self.accept())
-        button_box.rejected.connect(lambda: self.reject())
-        layout.addWidget(button_box)
+        self.button_box = QtWidgets.QDialogButtonBox(QtWidgets.QDialogButtonBox.Ok | QtWidgets.QDialogButtonBox.Cancel)
+        self.ok_btn = self.button_box.button(QtWidgets.QDialogButtonBox.StandardButton.Ok)
+        self.cancel_btn = self.button_box.button(QtWidgets.QDialogButtonBox.StandardButton.Cancel)
+
+        self.button_box.accepted.connect(lambda: self.accept())
+        self.button_box.rejected.connect(lambda: self.reject())
+        layout.addWidget(self.button_box)
 
         self.setLayout(layout)
+        self.refresh_validity()
 
     def add_name_validator(self, validator):
         if self.name_field:
             self.name_field.setValidator(validator)
 
+    def refresh_validity(self):
+        name_valid = self.name_field.hasAcceptableInput()
+        addr_valid = self.address_field.hasAcceptableInput()
+
+        sb_widgets.set_qt_property(self.name_field, 'validation', 'acceptable' if name_valid else 'intermediate')
+        sb_widgets.set_qt_property(self.address_field, 'validation', 'acceptable' if addr_valid else 'intermediate')
+
+        self.ok_btn.setEnabled(name_valid and addr_valid)
+
     def devices_to_add(self):
         return [{"type": self.device_type, "name": self.name_field.text(), "address": self.address_field.text(), "kwargs": {}}]
 
