Switchboard: Update device widget inline `on_name_changed` handler to account for QValidator Invalid -> Intermediate.

Followup to CL 46138126.

diff --git a/Engine/Plugins/VirtualProduction/Switchboard/Source/Switchboard/switchboard/devices/device_widget_base.py b/Engine/Plugins/VirtualProduction/Switchboard/Source/Switchboard/switchboard/devices/device_widget_base.py
--- a/Engine/Plugins/VirtualProduction/Switchboard/Source/Switchboard/switchboard/devices/device_widget_base.py
+++ b/Engine/Plugins/VirtualProduction/Switchboard/Source/Switchboard/switchboard/devices/device_widget_base.py
@@ -277,7 +277,8 @@ class DeviceWidget(QtWidgets.QWidget):
     def on_name_changed(self, text):
         if not self.name_validator:
             return
-        if self.name_validator.validate(text, self.name_line_edit.cursorPosition()) == QtGui.QValidator.Invalid:
+
+        if self.name_validator.validate(text, self.name_line_edit.cursorPosition()) != QtGui.QValidator.State.Acceptable:
             rect = self.name_line_edit.parent().mapToGlobal(self.name_line_edit.geometry().topRight())
             self.help_tool_tip.showText(rect, "Names must be unique")
 
