Switchboard: Specify Unreal OSC server address/port via command line.

diff --git a/Engine/Plugins/VirtualProduction/Switchboard/Source/Switchboard/switchboard/devices/unreal/plugin_unreal.py b/Engine/Plugins/VirtualProduction/Switchboard/Source/Switchboard/switchboard/devices/unreal/plugin_unreal.py
--- a/Engine/Plugins/VirtualProduction/Switchboard/Source/Switchboard/switchboard/devices/unreal/plugin_unreal.py
+++ b/Engine/Plugins/VirtualProduction/Switchboard/Source/Switchboard/switchboard/devices/unreal/plugin_unreal.py
@@ -2489,6 +2489,10 @@ class DeviceUnreal(Device):
         )
         command_line_args += ini_engine
 
+        if self.is_recording_device:
+            command_line_args += f' -VPOSCServerAddress={self.address}'
+            command_line_args += f' -VPOSCServerPort={self.device_osc_port}'
+
         # Logging
         logging_args = DeviceUnreal.csettings['logging'].get_command_line_arg(
             override_device_name=self.name)
