Switchboard: Add support for suppressing LLH crash recovery, make Unreal OSC endpoint args unconditional.

diff --git a/Engine/Plugins/VirtualProduction/Switchboard/Source/Switchboard/switchboard/devices/livelinkhub/plugin_livelinkhub.py b/Engine/Plugins/VirtualProduction/Switchboard/Source/Switchboard/switchboard/devices/livelinkhub/plugin_livelinkhub.py
--- a/Engine/Plugins/VirtualProduction/Switchboard/Source/Switchboard/switchboard/devices/livelinkhub/plugin_livelinkhub.py
+++ b/Engine/Plugins/VirtualProduction/Switchboard/Source/Switchboard/switchboard/devices/livelinkhub/plugin_livelinkhub.py
@@ -56,6 +56,13 @@ class DeviceLiveLinkHub(DeviceUnreal):
             tool_tip='When checked, retrieves logs after Live Link Hub terminates',
             category='General Settings',
         ),
+        'disable_crash_recovery': BoolSetting(
+            attr_name='disable_crash_recovery',
+            nice_name='Disable Crash Recovery',
+            value=False,
+            tool_tip='Stops Live Link Hub from prompting to recover after an unclean shutdown.',
+            category='General Settings',
+        ),
         'port': DeviceUnreal.csettings['port'],
     }
 
@@ -171,6 +178,13 @@ class DeviceLiveLinkHub(DeviceUnreal):
         if self.session_path:
             command_line_args += f' -SessionPath="{self.session_path}"'
 
+        if DeviceLiveLinkHub.csettings['disable_crash_recovery'].get_value():
+            command_line_args += (
+                ' -ini:Engine:'
+                '[/Script/LiveLinkHub.LiveLinkHubSettings]:'
+                'bEnableCrashRecovery=false'
+            )
+
         return command_line_args
 
     #@override
diff --git a/Engine/Plugins/VirtualProduction/Switchboard/Source/Switchboard/switchboard/devices/unreal/plugin_unreal.py b/Engine/Plugins/VirtualProduction/Switchboard/Source/Switchboard/switchboard/devices/unreal/plugin_unreal.py
--- a/Engine/Plugins/VirtualProduction/Switchboard/Source/Switchboard/switchboard/devices/unreal/plugin_unreal.py
+++ b/Engine/Plugins/VirtualProduction/Switchboard/Source/Switchboard/switchboard/devices/unreal/plugin_unreal.py
@@ -2489,9 +2489,8 @@ class DeviceUnreal(Device):
         )
         command_line_args += ini_engine
 
-        if self.is_recording_device:
-            command_line_args += f' -VPOSCServerAddress={self.address}'
-            command_line_args += f' -VPOSCServerPort={self.device_osc_port}'
+        command_line_args += f' -VPOSCServerAddress={self.address}'
+        command_line_args += f' -VPOSCServerPort={self.device_osc_port}'
 
         # Logging
         logging_args = DeviceUnreal.csettings['logging'].get_command_line_arg(
