Switchboard: Disable base class validation in nDisplay device creation dialog.

diff --git a/Engine/Plugins/VirtualProduction/Switchboard/Source/Switchboard/switchboard/devices/ndisplay/plugin_ndisplay.py b/Engine/Plugins/VirtualProduction/Switchboard/Source/Switchboard/switchboard/devices/ndisplay/plugin_ndisplay.py
--- a/Engine/Plugins/VirtualProduction/Switchboard/Source/Switchboard/switchboard/devices/ndisplay/plugin_ndisplay.py
+++ b/Engine/Plugins/VirtualProduction/Switchboard/Source/Switchboard/switchboard/devices/ndisplay/plugin_ndisplay.py
@@ -154,6 +154,11 @@ class AddnDisplayDialog(AddDeviceDialog):
         # populate the config combobox with the items last populated.
         self.recall_config_itemDatas()
 
+    #@override
+    def refresh_validity(self):
+        # Bypass base class validation for fields we don't use
+        pass
+
     def on_cmbLaunchAs_currentIndexChanged(self, state):
         ''' Called when the user changes the state of the LaunchAs dropdown
         It updates the editability of the packaged game path selection accordingly.
