Switchboard: Syncstatus fixes for "Engine Sync Method" == "Use Existing", sync filters.

diff --git a/Engine/Plugins/VirtualProduction/Switchboard/Source/Switchboard/sbl_helper.py b/Engine/Plugins/VirtualProduction/Switchboard/Source/Switchboard/sbl_helper.py
--- a/Engine/Plugins/VirtualProduction/Switchboard/Source/Switchboard/sbl_helper.py
+++ b/Engine/Plugins/VirtualProduction/Switchboard/Source/Switchboard/sbl_helper.py
@@ -769,20 +769,22 @@ class SbListenerHelper:
         remove_excluded_workload = SyncWorkload()
         if options.remove_excluded_workspace_files and self.ugs_filters:
             logging.info('Finding files in workspace that need to be removed.')
-            haves = p4_utils.p4_have(['//...'], **p4kwargs)
-            stream = self.p4clientspec['Stream']
-            stream_path = pathlib.PurePosixPath(stream)
-            num_exclusions = 0
-            for have in haves:
-                client_str = have['path']
-                depot_str = have['depotFile']
-                depot_path = pathlib.PurePosixPath(depot_str)
-                rel_path = f'/{depot_path.relative_to(stream_path)}'
-                if not self.ugs_filters.includes_path(rel_path):
-                    num_exclusions += 1
-                    remove_excluded_workload.add(
-                        SyncFile(depot_str, client_str, 0, 0, 'deleted'))
-            logging.info(f'Found {num_exclusions:,} excluded files')
+            haves = p4_utils.p4_have([*engine_paths, *project_paths], **p4kwargs)
+            if stream := self.p4clientspec.get('Stream'):
+                stream_path = pathlib.PurePosixPath(stream)
+                num_exclusions = 0
+                for have in haves:
+                    client_str = have['path']
+                    depot_str = have['depotFile']
+                    depot_path = pathlib.PurePosixPath(depot_str)
+                    rel_path = f'/{depot_path.relative_to(stream_path)}'
+                    if not self.ugs_filters.includes_path(rel_path):
+                        num_exclusions += 1
+                        remove_excluded_workload.add(
+                            SyncFile(depot_str, client_str, 0, 0, 'deleted'))
+                logging.info(f'Found {num_exclusions:,} excluded files')
+            else:
+                logging.error('Unable to determine stream from client spec')
 
         def sync_filter(workload: SyncWorkload):
             # Don't add Build.version to the workspace, but do allow removal.
@@ -794,19 +796,21 @@ class SbListenerHelper:
             if self.ugs_filters and self.ugs_filters.excluded_categories:
                 num_files = len(workload.sync_files)
                 logging.info(f'Applying sync filters... ({num_files:,} files)')
-                stream = self.p4clientspec['Stream']
-                stream_path = pathlib.PurePosixPath(stream)
-                excluded_files: set[str] = set()
-                for depot_str, syncfile in workload.sync_files.items():
-                    depot_path = pathlib.PurePosixPath(depot_str)
-                    depot_rel_path = f'/{depot_path.relative_to(stream_path)}'
-                    if not self.ugs_filters.includes_path(depot_rel_path):
-                        excluded_files.add(depot_str)
-
-                num_exclusions = len(excluded_files)
-                logging.info(f'Filter excluded {num_exclusions:,} files')
-                for exclusion in excluded_files:
-                    workload.sync_files.pop(exclusion)
+                if stream := self.p4clientspec.get('Stream'):
+                    stream_path = pathlib.PurePosixPath(stream)
+                    excluded_files: set[str] = set()
+                    for depot_str, syncfile in workload.sync_files.items():
+                        depot_path = pathlib.PurePosixPath(depot_str)
+                        depot_rel_path = f'/{depot_path.relative_to(stream_path)}'
+                        if not self.ugs_filters.includes_path(depot_rel_path):
+                            excluded_files.add(depot_str)
+
+                    num_exclusions = len(excluded_files)
+                    logging.info(f'Filter excluded {num_exclusions:,} files')
+                    for exclusion in excluded_files:
+                        workload.sync_files.pop(exclusion)
+                else:
+                    logging.error('Unable to determine stream from client spec')
 
             for removal in remove_excluded_workload.sync_files.values():
                 workload.add_or_replace(removal)
@@ -1052,6 +1056,9 @@ class SbListenerHelper:
     def run_syncstatus(self, options: argparse.Namespace) -> int:
         self.check_p4_options(options)
 
+        # syncstatus only considers engine_dir if it was passed in, NOT if we auto-detected it.
+        specified_engine_dir = options.engine_dir is not None
+
         def guess_have_cl(path: pathlib.Path, client: str | None = None) -> int | None:
             logging.debug(f'Falling back to `p4 cstat #have` for {path}')
 
@@ -1080,7 +1087,8 @@ class SbListenerHelper:
 
         result = {}
 
-        if self.engine_dir:
+        if specified_engine_dir:
+            assert self.engine_dir
             if eng_detected_ws := p4_utils.p4_get_var('P4CLIENT', self.engine_dir):
                 result['engine_detected_ws'] = eng_detected_ws
 
@@ -1118,15 +1126,24 @@ class SbListenerHelper:
             orig_cwd = os.getcwd()  # Ensure user's P4CONFIG files are read appropriately
             if self.p4client_project and self.p4client_project != self.p4client:
                 # Potentially two queries if we specify a separate project workspace.
-                either_dirty = False
+                engine_dirty = False
+                project_dirty = False
 
-                if self.engine_dir:
+                if specified_engine_dir:
+                    assert self.engine_dir
                     os.chdir(self.engine_dir)
                     eng_workload, _ = self._get_sync_workload(options,
                         engine_dir=self.engine_dir,
                         sync_engine_cl=self.sync_engine_cl,
                         port=self.p4port, user=self.p4user, client=self.p4client)
-                    either_dirty = either_dirty or len(eng_workload.sync_files) > 0
+                    if len(eng_workload.sync_files) > 0:
+                        engine_dirty = True
+                        logging.debug('Engine would sync:')
+                        for idx, sync_file in enumerate(eng_workload.sync_files):
+                            logging.debug(f'\t- {sync_file}')
+                            if idx >= 100:
+                                logging.debug(f'\t...and {len(eng_workload.sync_files) - idx} more')
+                                break
 
                 if self.project_dir:
                     os.chdir(self.project_dir)
@@ -1134,9 +1151,16 @@ class SbListenerHelper:
                         project_dir=self.project_dir,
                         sync_project_cl=self.sync_project_cl,
                         port=self.p4port, user=self.p4user, client=self.p4client_project)
-                    either_dirty = either_dirty or len(proj_workload.sync_files) > 0
-
-                result['dirty'] = either_dirty
+                    if len(proj_workload.sync_files) > 0:
+                        project_dirty = True
+                        logging.debug('Project would sync:')
+                        for idx, sync_file in enumerate(proj_workload.sync_files):
+                            logging.debug(f'\t- {sync_file}')
+                            if idx >= 100:
+                                logging.debug(f'\t...and {len(proj_workload.sync_files) - idx} more')
+                                break
+
+                result['dirty'] = engine_dirty or project_dirty
             else:
                 sync_workload, _ = self._get_sync_workload(options,
                     engine_dir=self.engine_dir,
@@ -1145,7 +1169,17 @@ class SbListenerHelper:
                     sync_project_cl=self.sync_project_cl,
                     **self.p4kwargs)
 
-                result['dirty'] = len(sync_workload.sync_files) > 0
+                if len(sync_workload.sync_files) > 0:
+                    result['dirty'] = True
+                    logging.debug('Would sync:')
+                    for idx, sync_file in enumerate(sync_workload.sync_files):
+                        logging.debug(f'\t- {sync_file}')
+                        if idx >= 100:
+                            logging.debug(f'\t...and {len(sync_workload.sync_files) - idx} more')
+                            break
+                else:
+                    result['dirty'] = False
+
         finally:
             os.chdir(orig_cwd)
 
@@ -1365,6 +1399,7 @@ class SbListenerHelper:
 
         if self.p4client:
             self.p4kwargs['client'] = self.p4client
+            self.p4clientspec = p4_utils.p4_get_client(self.p4client, **self.p4kwargs) or {}
 
         if p4login := p4_utils.p4_get_login(**self.p4kwargs):
             logging.info(f'Perforce: Logged in as {p4login}')
@@ -1389,9 +1424,6 @@ class SbListenerHelper:
         self.sync_pcbs = options.use_pcbs
         self.ugs_lib_dir = options.ugs_lib_dir
 
-        if self.p4client:
-            self.p4clientspec = p4_utils.p4_get_client(self.p4client, **self.p4kwargs) or {}
-
         if options.project is None:
             if options.project_cl is not None:
                 self.parser.error('--project-cl requires --project')
diff --git a/Engine/Plugins/VirtualProduction/Switchboard/Source/Switchboard/switchboard/devices/unreal/plugin_unreal.py b/Engine/Plugins/VirtualProduction/Switchboard/Source/Switchboard/switchboard/devices/unreal/plugin_unreal.py
--- a/Engine/Plugins/VirtualProduction/Switchboard/Source/Switchboard/switchboard/devices/unreal/plugin_unreal.py
+++ b/Engine/Plugins/VirtualProduction/Switchboard/Source/Switchboard/switchboard/devices/unreal/plugin_unreal.py
@@ -1616,9 +1616,7 @@ class DeviceUnreal(Device):
 
         syncstatus_args = [
             f'"{self.sbl_helper_path}"',
-            # '--log-level=DEBUG',
             'syncstatus',
-            f'--engine-dir="{engine_dir}"',
             f'--project="{project_path}"',
         ]
 
@@ -1631,6 +1629,20 @@ class DeviceUnreal(Device):
         if project_workspace:
             syncstatus_args.append(f'--p4client-project={project_workspace}')
 
+        if DeviceUnreal.csettings['use_sync_filters'].get_value():
+            include_str = ''
+            custom_view_str = self.csettings['custom_sync_view'].get_value()
+
+            incl_categories: list[SyncCategoryOption] = self.csettings[
+                'included_sync_categories'].get_value()
+            if incl_categories:
+                incl_categories.sort(key=lambda cat: str(cat.catid))
+                include_str = ','.join(str(cat.catid) for cat in incl_categories)
+                syncstatus_args.append(f'--include-categories="{include_str}"')
+
+            if custom_view_str:
+                syncstatus_args.append(f'--custom-view="{custom_view_str}"')
+
         program_name = UnrealJobs.SyncStatus.value
 
         working_dir = os.path.dirname(CONFIG.UPROJECT_PATH.get_value(self.name))
@@ -1813,7 +1825,7 @@ class DeviceUnreal(Device):
             incl_categories: list[SyncCategoryOption] = self.csettings[
                 'included_sync_categories'].get_value()
             if incl_categories:
-                incl_categories.sort(key=lambda cat: cat.catid)
+                incl_categories.sort(key=lambda cat: str(cat.catid))
                 include_str = ','.join(str(cat.catid) for cat in incl_categories)
                 sync_args += f' --include-categories="{include_str}"'
 
