Switchboard: Fix for "Invalid engine changelist specified for sync" when "Engine Sync Method" has been changed from "Build Engine" to "Use Existing" without restarting Switchboard.

diff --git a/Engine/Plugins/VirtualProduction/Switchboard/Source/Switchboard/switchboard/devices/unreal/plugin_unreal.py b/Engine/Plugins/VirtualProduction/Switchboard/Source/Switchboard/switchboard/devices/unreal/plugin_unreal.py
--- a/Engine/Plugins/VirtualProduction/Switchboard/Source/Switchboard/switchboard/devices/unreal/plugin_unreal.py
+++ b/Engine/Plugins/VirtualProduction/Switchboard/Source/Switchboard/switchboard/devices/unreal/plugin_unreal.py
@@ -1724,6 +1724,12 @@ class DeviceUnreal(Device):
         project_cl: Optional[int | str],
         puuid_dependency: Optional[uuid.UUID] = None
     ) -> Optional[uuid.UUID]:
+        if engine_cl == '':
+            engine_cl = None
+
+        if project_cl == '':
+            project_cl = None
+
         if not engine_cl and not project_cl:
             LOGGER.warning(
                 "Neither project nor engine changelist is selected. "
