[Switchboard] Fixed adding nDisplay perf test devices
- Now the "Ok" button appears active once test file is selected

#rnx

diff --git a/Engine/Plugins/VirtualProduction/Switchboard/Source/Switchboard/switchboard/devices/ndisplayTest/plugin_ndisplayTest.py b/Engine/Plugins/VirtualProduction/Switchboard/Source/Switchboard/switchboard/devices/ndisplayTest/plugin_ndisplayTest.py
--- a/Engine/Plugins/VirtualProduction/Switchboard/Source/Switchboard/switchboard/devices/ndisplayTest/plugin_ndisplayTest.py
+++ b/Engine/Plugins/VirtualProduction/Switchboard/Source/Switchboard/switchboard/devices/ndisplayTest/plugin_ndisplayTest.py
@@ -30,11 +30,18 @@ from .ndisplayTest import nDisplayTest
 class AddnDisplayTestDialog(AddDeviceDialog):
     
     def __init__(self, existing_devices, parent=None):
+
+        # Initialization flag to ignore any virtual calls during initialization
+        self.baseInitializing = True
+
         super().__init__(
             device_type="nDisplayTest",
             existing_devices=existing_devices,
             parent=parent)
 
+        # Now clear the guard flag
+        self.baseInitializing = False
+
         # Initialize members
         self.existing_devices = existing_devices
         self.test_dir = self.get_default_test_dir()
@@ -43,6 +50,10 @@ class AddnDisplayTestDialog(AddDeviceDialog):
         # Initialize GUI
         self.initialize_widgets()
 
+        # Update "Ok" button state
+        self.refresh_validity()
+
+
     def initialize_widgets(self):
         ''' Performs widgets initialization '''
 
@@ -157,6 +168,9 @@ class AddnDisplayTestDialog(AddDeviceDialog):
             dir_new = os.path.normpath(dir_new)
             if dir_old != dir_new:
                 self.set_test_location(dir_new)
+        
+        # Update "Ok" button state
+        self.refresh_validity()
 
     def get_selected_test_file(self):
         ''' Returns the test chosen '''
@@ -213,6 +227,16 @@ class AddnDisplayTestDialog(AddDeviceDialog):
         # Remove all existing devices
         return self.existing_devices
 
+    def refresh_validity(self):
+        ''' [override] Updates the Ok button state '''
+
+        # Ignore calls from the base constructor
+        if self.baseInitializing:
+            return
+
+        # Activate the Ok button if any test selected, otherwise deactivate
+        self.ok_btn.setEnabled(bool(self.cbTests.currentText().strip()))
+
 
 ####################################################
 # nDisplay Test device widget
