Switchboard: Fix for incorrectly displaying "not synced" state in some circumstances.

diff --git a/Engine/Plugins/VirtualProduction/Switchboard/Source/Switchboard/switchboard/devices/unreal/plugin_unreal.py b/Engine/Plugins/VirtualProduction/Switchboard/Source/Switchboard/switchboard/devices/unreal/plugin_unreal.py
--- a/Engine/Plugins/VirtualProduction/Switchboard/Source/Switchboard/switchboard/devices/unreal/plugin_unreal.py
+++ b/Engine/Plugins/VirtualProduction/Switchboard/Source/Switchboard/switchboard/devices/unreal/plugin_unreal.py
@@ -4059,7 +4059,7 @@ class DeviceWidgetUnreal(DeviceWidget):
         except (TypeError, ValueError):
             self.sync_button.hide()
 
-        is_synched = required_cl is None or required_cl == current_device_cl
+        is_synched = required_cl is None or str(required_cl) == str(current_device_cl)
         self._set_project_changelist_is_synched(is_synched)
 
     def update_build_status(self, device, step, percent):
@@ -4115,7 +4115,7 @@ class DeviceWidgetUnreal(DeviceWidget):
         except (TypeError, ValueError):
             self.sync_button.hide()
 
-        is_synched = required_cl is None or required_cl == synced_cl
+        is_synched = required_cl is None or str(required_cl) == str(synced_cl)
         self._set_engine_changelist_is_synched(is_synched)
         self.update_build_info(device)
 
