Switchboard: Live Link Hub device improvements.

- Change topology toggle from gold star to dropdown with icons.
- Rename user-facing references to "session" to "config" for consistency.
- Filter browse dialog to *.json by default.
- Allow config to be overridden per-device.

diff --git a/Engine/Plugins/VirtualProduction/Switchboard/Source/Switchboard/switchboard/devices/livelinkhub/plugin_livelinkhub.py b/Engine/Plugins/VirtualProduction/Switchboard/Source/Switchboard/switchboard/devices/livelinkhub/plugin_livelinkhub.py
--- a/Engine/Plugins/VirtualProduction/Switchboard/Source/Switchboard/switchboard/devices/livelinkhub/plugin_livelinkhub.py
+++ b/Engine/Plugins/VirtualProduction/Switchboard/Source/Switchboard/switchboard/devices/livelinkhub/plugin_livelinkhub.py
@@ -2,33 +2,29 @@
 
 from __future__ import annotations
 
+from enum import IntEnum
 import pathlib
 from typing import Callable
 from uuid import UUID
 
 from PySide6 import QtCore, QtWidgets
 
-from switchboard import message_protocol, switchboard_widgets
+from switchboard import message_protocol
 from switchboard.config import CONFIG, Setting, BoolSetting, FilePathSetting, StringSetting
 from switchboard.devices.device_base import Device, DeviceStatus, PluginHeaderWidgets
-from switchboard.devices.unreal.plugin_unreal import (DeviceUnreal, DeviceWidgetUnreal, ProgramStartQueue,
-                                                      ProgramStartQueueItem)
+from switchboard.devices.unreal.plugin_unreal import DeviceUnreal, DeviceWidgetUnreal, ProgramStartQueueItem
 from switchboard.switchboard_logging import LOGGER
 
 
+class LiveLinkHubTopologyMode(IntEnum):
+    Hub = 0
+    Spoke = 1
+
+
 class DeviceLiveLinkHub(DeviceUnreal):
     LLH_PROG_NAME = 'livelinkhub'
 
-    all_llh_devices: set[DeviceLiveLinkHub] = set()
-
     csettings: dict[str, Setting] = {
-        'primary_device_name': StringSetting(
-            attr_name='primary_device_name',
-            nice_name='Primary Device',
-            value='',
-            tool_tip='Identifies which Live Link Hub device should operate in "hub" (non-spoke) mode',
-            show_ui=False,
-        ),
         'executable_name': StringSetting(
             attr_name='executable_name',
             nice_name='Live Link Hub executable filename',
@@ -44,10 +40,11 @@ class DeviceLiveLinkHub(DeviceUnreal):
         ),
         'session_path': FilePathSetting(
             attr_name='session_path',
-            nice_name='Session File Path',
+            nice_name='Config File Path',
             value='',
             tool_tip='Path to the Live Link Hub session file to load on startup',
             category='General Settings',
+            file_path_filter="JSON files (*.json);;All Files (*)",
         ),
         'retrieve_logs': BoolSetting(
             attr_name='retrieve_logs',
@@ -70,6 +67,12 @@ class DeviceLiveLinkHub(DeviceUnreal):
         super().__init__(name, address, *args, **kwargs)
 
         self.widget: DeviceWidgetLiveLinkHub
+        self.setting_spoke_mode = BoolSetting(
+            attr_name='spoke_mode',
+            nice_name='Spoke Mode',
+            value=kwargs.pop('spoke_mode', False),
+            tool_tip='Whether to start this instance of Live Link Hub in "spoke" topology mode',
+        )
 
     @property
     def llh_proj_dir(self) -> pathlib.Path:
@@ -77,24 +80,6 @@ class DeviceLiveLinkHub(DeviceUnreal):
         engine_dir = pathlib.Path(CONFIG.ENGINE_DIR.get_value(self.name))
         return engine_dir / 'Source' / 'Programs' / 'LiveLinkHub'
 
-    @classmethod
-    def select_device_as_primary(cls, selected_device: DeviceLiveLinkHub):
-        ''' Selects the given device as the "hub" (non-spoke) device '''
-
-        cls.csettings['primary_device_name'].update_value(selected_device.name)
-
-        for dev in cls.all_llh_devices:
-            if dev.widget:
-                dev.widget.primary_button.setChecked(dev.name == selected_device.name)
-
-    def is_primary(self):
-        ''' Check if the device is the "hub" (non-spoke) device or not '''
-        return self.name == DeviceLiveLinkHub.csettings['primary_device_name'].get_value()
-
-    def select_as_primary(self):
-        ''' Selects this device as the "hub" (non-spoke) device '''
-        DeviceLiveLinkHub.select_device_as_primary(self)
-
     #@override
     @classmethod
     def plugin_settings(cls):
@@ -103,12 +88,14 @@ class DeviceLiveLinkHub(DeviceUnreal):
     #@override
     def device_settings(self):
         return Device.device_settings(self) + [
+            self.setting_spoke_mode
         ]
 
     #@override
     def setting_overrides(self):
         return [
             DeviceLiveLinkHub.csettings['command_line_arguments'],
+            DeviceLiveLinkHub.csettings['session_path'],
             CONFIG.ENGINE_DIR,
             CONFIG.SOURCE_CONTROL_WORKSPACE,
         ]
@@ -118,32 +105,22 @@ class DeviceLiveLinkHub(DeviceUnreal):
     def plugin_header_widget_config(cls):
         return super().plugin_header_widget_config() & ~PluginHeaderWidgets.AUTOJOIN_MU
 
-    #@override
-    @classmethod
-    def added_device(cls, device: DeviceLiveLinkHub):
-        super().added_device(device)
-
-        assert device not in DeviceLiveLinkHub.all_llh_devices
-        DeviceLiveLinkHub.all_llh_devices.add(device)
-
-        if not cls.csettings['primary_device_name'].get_value():
-            # Automatically make the first added device "primary"
-            cls.select_device_as_primary(device)
-
-    #@override
-    @classmethod
-    def removed_device(cls, device: DeviceLiveLinkHub):
-        super().removed_device(device)
-
-        assert device in DeviceLiveLinkHub.all_llh_devices
-        DeviceLiveLinkHub.all_llh_devices.remove(device)
-
     #@override
     def device_widget_registered(self, device_widget: DeviceWidgetLiveLinkHub):
         super().device_widget_registered(device_widget)
 
-        device_widget.signal_device_widget_primary.connect(self.select_as_primary)
-        device_widget.primary_button.setChecked(self.is_primary())
+        # Reflect current mode in the widget
+        if self.setting_spoke_mode.get_value():
+            current_mode = LiveLinkHubTopologyMode.Spoke
+        else:
+            current_mode = LiveLinkHubTopologyMode.Hub
+
+        device_widget.mode_dropdown.setCurrentIndex(current_mode)
+
+        # Respond to mode changes via the widget
+        device_widget.signal_device_widget_spoke_mode.connect(
+            lambda new_is_spoke: self.setting_spoke_mode.update_value(new_is_spoke)
+        )
 
     #@override
     @property
@@ -169,7 +146,7 @@ class DeviceLiveLinkHub(DeviceUnreal):
 
         command_line_args += f' Log={self.log_filename}'
 
-        if self.is_primary():
+        if self.setting_spoke_mode.get_value() == False:
             command_line_args += ' -hub'
             command_line_args += ' -UDPMESSAGING_SHARE_KNOWN_NODES=1'
         else:
@@ -269,12 +246,12 @@ class DeviceLiveLinkHub(DeviceUnreal):
 
 
 class DeviceWidgetLiveLinkHub(DeviceWidgetUnreal):
-    signal_device_widget_primary = QtCore.Signal(object)
+    signal_device_widget_spoke_mode = QtCore.Signal(bool)
 
     def __init__(self, *args, **kwargs):
         super().__init__(*args, **kwargs)
 
-        self.primary_button: switchboard_widgets.ControlQPushButton
+        self.mode_dropdown: QtWidgets.QComboBox
 
     #@override
     @property
@@ -289,17 +266,16 @@ class DeviceWidgetLiveLinkHub(DeviceWidgetUnreal):
     #@override
     def add_widget_to_layout(self, widget):
         if widget == self.name_line_edit:
-            self._add_primary_button()
+            self._add_mode_dropdown()
 
-            # shorten the widget to account for the inserted primary button
-            btn_added_width = (
-                max(self.primary_button.iconSize().width(),
-                    self.primary_button.minimumSize().width())
+            # shorten the widget to account for the inserted dropdown
+            added_width = (
+                self.mode_dropdown.minimumSize().width()
                 + 2 * self.layout.spacing()
-                + self.primary_button.contentsMargins().left()
-                + self.primary_button.contentsMargins().right())
+                + self.mode_dropdown.contentsMargins().left()
+                + self.mode_dropdown.contentsMargins().right())
 
-            le_maxwidth = self.name_line_edit.maximumWidth() - btn_added_width
+            le_maxwidth = self.name_line_edit.maximumWidth() - added_width
             self.name_line_edit.setMaximumWidth(le_maxwidth)
 
         super().add_widget_to_layout(widget)
@@ -311,26 +287,32 @@ class DeviceWidgetLiveLinkHub(DeviceWidgetUnreal):
         cmenu.addAction('Open fetched log', lambda: self.signal_open_last_log.emit(self))
         cmenu.addAction('Copy last launch command', lambda: self.signal_copy_last_launch_command.emit(self))
 
-    def _add_primary_button(self):
+    #@override
+    def _update_connected_ui(self):
+        super()._update_connected_ui()
+        self.mode_dropdown.setDisabled(False)
+
+    #@override
+    def _update_disconnected_ui(self):
+        super()._update_disconnected_ui()
+        self.mode_dropdown.setDisabled(True)
+
+    def _add_mode_dropdown(self):
         '''
-        Adds to the layout a button to select which device should be the
-        "hub"-mode device; all others are inferred to be "spoke" mode.
+        Adds to the layout a dropdown to select which topology mode the device should be in.
         '''
 
-        self.primary_button = switchboard_widgets.ControlQPushButton.create(
-            ':/icons/images/star_yellow_off.png',
-            icon_disabled=':/icons/images/star_yellow_off.png',
-            icon_hover=':/icons/images/star_yellow_off.png',
-            icon_disabled_on=':/icons/images/star_yellow_off.png',
-            icon_on=':/icons/images/star_yellow.png',
-            icon_hover_on=':/icons/images/star_yellow.png',
-            icon_size=QtCore.QSize(16, 16),
-            tool_tip='Select as "hub" (non-spoke) device'
-        )
+        self.mode_dropdown = QtWidgets.QComboBox(self)
+        self.mode_dropdown.setFixedSize(40, 32)
+        self.mode_dropdown.view().setFixedWidth(100)
+        for mode in LiveLinkHubTopologyMode:
+            self.mode_dropdown.addItem(mode.name)
+            self.mode_dropdown.setItemData(mode.value, self.icons[f'{mode.name}Mode'],
+                                           QtCore.Qt.ItemDataRole.DecorationRole)
 
-        self.add_widget_to_layout(self.primary_button)
+        self.add_widget_to_layout(self.mode_dropdown)
 
-        self.primary_button.clicked.connect(self._on_primary_button_clicked)
+        self.mode_dropdown.currentIndexChanged.connect(self._on_mode_dropdown_current_index_changed)
 
-    def _on_primary_button_clicked(self):
-        self.signal_device_widget_primary.emit(self)
+    def _on_mode_dropdown_current_index_changed(self, idx: int):
+        self.signal_device_widget_spoke_mode.emit(idx == LiveLinkHubTopologyMode.Spoke)
