Switchboard: Persist console exec history.

diff --git a/Engine/Plugins/VirtualProduction/Switchboard/Source/Switchboard/switchboard/devices/ndisplay/ndisplay_monitor_ui.py b/Engine/Plugins/VirtualProduction/Switchboard/Source/Switchboard/switchboard/devices/ndisplay/ndisplay_monitor_ui.py
--- a/Engine/Plugins/VirtualProduction/Switchboard/Source/Switchboard/switchboard/devices/ndisplay/ndisplay_monitor_ui.py
+++ b/Engine/Plugins/VirtualProduction/Switchboard/Source/Switchboard/switchboard/devices/ndisplay/ndisplay_monitor_ui.py
@@ -1,8 +1,9 @@
 # Copyright Epic Games, Inc. All Rights Reserved.
 
-from typing import Dict, Optional
+from collections import OrderedDict
+from typing import Optional
 
-from PySide6.QtCore import QSortFilterProxyModel, Qt
+from PySide6.QtCore import Signal, QSortFilterProxyModel, Qt
 from PySide6.QtWidgets import QComboBox, QCompleter, QHBoxLayout, \
     QHeaderView, QLabel, QPushButton, QSizePolicy, QStyledItemDelegate, \
     QTableView, QVBoxLayout, QWidget, QCheckBox
@@ -12,14 +13,16 @@ from switchboard.devices.ndisplay.ndisplay_monitor import nDisplayMonitor
 
 class nDisplayMonitorUI(QWidget):
 
+    signal_exec_history_changed = Signal()
+
     def __init__(self, parent: QWidget, monitor: nDisplayMonitor):
         QWidget.__init__(self, parent)
 
         self.monitor = monitor
 
         # Console exec combo entries; key is lowercased, value is as-entered.
-        # TODO: Persist across launches? Frecency sorting?
-        self.exec_history: Dict[str, str] = {}
+        # TODO: Frecency sorting?
+        self.exec_history = OrderedDict[str, str]()
 
         self.cmb_console_exec = QComboBox()
         self.exec_model = QSortFilterProxyModel(self.cmb_console_exec)
@@ -149,13 +152,31 @@ class nDisplayMonitorUI(QWidget):
             self.cmb_console_exec.setCurrentIndex(-1)
 
     def update_exec_history(self, exec_str: str):
+        ''' Invoked when a command is executed, either adding or moving to most recent '''
         # Reinsert (case-insensitive) duplicates as most recent.
         exec_str_lower = exec_str.lower()
         if exec_str_lower in self.exec_history:
             del self.exec_history[exec_str_lower]
         self.exec_history[exec_str_lower] = exec_str
 
+        MAX_HISTORY_ENTRIES = 30
+        if len(self.exec_history) > MAX_HISTORY_ENTRIES:
+            self.exec_history.popitem(last=False)
+
+        self._rebuild_console_exec_items()
+
+        self.signal_exec_history_changed.emit()
+
+    def replace_exec_history(self, history: list[str]):
+        ''' Can be used to restore the history model from a persisted version. '''
+        self.exec_history = OrderedDict({
+            exec.lower(): exec
+            for exec in history
+        })
+
+        self._rebuild_console_exec_items()
+
+    def _rebuild_console_exec_items(self):
         # Most recently used at the top.
         self.cmb_console_exec.clear()
-        self.cmb_console_exec.addItems(
-            reversed(list(self.exec_history.values())))
+        self.cmb_console_exec.addItems(list(reversed(self.exec_history.values())))
diff --git a/Engine/Plugins/VirtualProduction/Switchboard/Source/Switchboard/switchboard/devices/ndisplay/plugin_ndisplay.py b/Engine/Plugins/VirtualProduction/Switchboard/Source/Switchboard/switchboard/devices/ndisplay/plugin_ndisplay.py
--- a/Engine/Plugins/VirtualProduction/Switchboard/Source/Switchboard/switchboard/devices/ndisplay/plugin_ndisplay.py
+++ b/Engine/Plugins/VirtualProduction/Switchboard/Source/Switchboard/switchboard/devices/ndisplay/plugin_ndisplay.py
@@ -712,6 +712,12 @@ class DevicenDisplay(DeviceUnreal):
             show_ui=True if sys.platform in ('win32', 'linux') else False,  # Gpu Clocker is available in select platforms
             category="GPU/CPU Settings",
         ),
+        'cluster_exec_history': StringListSetting(
+            attr_name='cluster_exec_history',
+            nice_name='Cluster exec console command history',
+            value=[],
+            show_ui=False,
+        )
     }
 
     ndisplay_monitor_ui = None
@@ -1898,12 +1904,19 @@ class DevicenDisplay(DeviceUnreal):
 
         # Create Monitor UI if it doesn't exist
         if not cls.ndisplay_monitor_ui:
-            cls.ndisplay_monitor_ui = nDisplayMonitorUI(
-                parent=tabs, monitor=cls.ndisplay_monitor)
+            cls.ndisplay_monitor_ui = nDisplayMonitorUI(parent=tabs, monitor=cls.ndisplay_monitor)
+            cls.ndisplay_monitor_ui.replace_exec_history(DevicenDisplay.csettings['cluster_exec_history'].get_value())
+            cls.ndisplay_monitor_ui.signal_exec_history_changed.connect(lambda: cls._handle_exec_history_changed())
 
         # Add our monitor UI to the main tabs in the UI
         tabs.addTab(cls.ndisplay_monitor_ui, 'nDisplay &Monitor')
 
+    @classmethod
+    def _handle_exec_history_changed(cls):
+        if cls.ndisplay_monitor_ui:
+            DevicenDisplay.csettings['cluster_exec_history'].update_value(
+                list(cls.ndisplay_monitor_ui.exec_history.values()))
+
     @classmethod
     def added_device(cls, device):
         '''
