Switchboard: Unreal devices are now launched with auto-save disabled.

Because auto-save typically interrupts stage operations.

Also re-organized DeviceUnreal.generate_unreal_command_line_args to be an array that later gets joined, making its use more consistent and easy. Also easier to debug.

#rnx

diff --git a/Engine/Plugins/VirtualProduction/Switchboard/Source/Switchboard/switchboard/devices/unreal/plugin_unreal.py b/Engine/Plugins/VirtualProduction/Switchboard/Source/Switchboard/switchboard/devices/unreal/plugin_unreal.py
--- a/Engine/Plugins/VirtualProduction/Switchboard/Source/Switchboard/switchboard/devices/unreal/plugin_unreal.py
+++ b/Engine/Plugins/VirtualProduction/Switchboard/Source/Switchboard/switchboard/devices/unreal/plugin_unreal.py
@@ -2362,27 +2362,53 @@ class DeviceUnreal(Device):
         return ocvars
 
     def generate_unreal_command_line_args(self, map_name):
-        command_line_args = f'{self.extra_cmdline_args_setting}'
 
-        command_line_args += f' Log={self.log_filename} '
+        # Build command line arguments as an array that will be joined at the end.
+        args = []
 
+        # DPCVars - accumulate these before adding to args
+        dp_cvars = []
+
+        # Project path and map
+        args.append(f'"{CONFIG.UPROJECT_PATH.get_value(self.name)}"')
+        if map_name:
+            args.append(map_name)
+
+        # Extra command line arguments from settings
+        if self.extra_cmdline_args_setting:
+            args.append(self.extra_cmdline_args_setting)
+
+        # Logging
+        args.append(f'Log={self.log_filename}')
+
+        # Multi-user server settings
         if CONFIG.MUSERVER_AUTO_JOIN.get_value() and self.autojoin_mu_server.get_value():
-            command_line_args += (
-                '-CONCERTRETRYAUTOCONNECTONERROR '
-                '-CONCERTAUTOCONNECT ')
+            args.extend([
+                '-CONCERTRETRYAUTOCONNECTONERROR',
+                '-CONCERTAUTOCONNECT'
+            ])
+
         mu_server = switchboard_application.get_multi_user_server_instance()
-        command_line_args += (f'-CONCERTSERVER="{mu_server.configured_server_name()}" '
-                              f'-CONCERTSESSION="{CONFIG.MUSERVER_SESSION_NAME.get_value()}" '
-                              f'-CONCERTDISPLAYNAME="{self.name}"')
+        args.extend([
+            f'-CONCERTSERVER="{mu_server.configured_server_name()}"',
+            f'-CONCERTSESSION="{CONFIG.MUSERVER_SESSION_NAME.get_value()}"',
+            f'-CONCERTDISPLAYNAME="{self.name}"'
+        ])
 
+        # Insights tracing
         if CONFIG.INSIGHTS_TRACE_ENABLE.get_value() and not self.exclude_from_insights.get_value():
             LOGGER.warning(f"Unreal Insight Tracing is enabled for '{self.name}'. This may affect Unreal Engine performance.")
             remote_utrace_path = self.get_utrace_filepath()
-            command_line_args += ' -statnamedevents' if CONFIG.INSIGHTS_STAT_EVENTS.get_value() else ''
-            command_line_args += ' -tracefile="{}" -trace="{}"'.format(
-                remote_utrace_path,
-                CONFIG.INSIGHTS_TRACE_ARGS.get_value())
 
+            if CONFIG.INSIGHTS_STAT_EVENTS.get_value():
+                args.append('-statnamedevents')
+
+            args.extend([
+                f'-tracefile="{remote_utrace_path}"',
+                f'-trace="{CONFIG.INSIGHTS_TRACE_ARGS.get_value()}"'
+            ])
+
+        # Exec commands
         exec_cmds = DeviceUnreal.csettings["exec_cmds"].get_value(self.name).copy()
 
         # LiveLink preset
@@ -2390,20 +2416,21 @@ class DeviceUnreal(Device):
         if livelink_preset_gamepath:
             exec_cmds.append(self.exec_command_for_livelink_preset(livelink_preset_gamepath))
 
-        # Exec Commands
-        exec_cmds = [cmd for cmd in exec_cmds if len(cmd.strip())]
+        # Media profile
+        mediaprofile_gamepath = DeviceUnreal.csettings["mediaprofile"].get_value(self.name)
+        if mediaprofile_gamepath:
+            dp_cvars.append(self.dpcvar_for_mediaprofile(mediaprofile_gamepath))
 
+        # Filter and format exec commands
+        exec_cmds = [cmd for cmd in exec_cmds if len(cmd.strip())]
         if len(exec_cmds):
             exec_cmds_expanded = ','.join(exec_cmds)
-            command_line_args += f' -ExecCmds="{exec_cmds_expanded}"'
-
-        # DPCVars may need to be appended to, so we don't concatenate them until the end.
-        dp_cvars = []
+            args.append(f'-ExecCmds="{exec_cmds_expanded}"')
 
+        # VP Roles
         (supported_roles, unsupported_roles) = self.get_vproles()
-
         if supported_roles or DeviceUnreal.any_devices_have_roles():
-            command_line_args += ' -VPRole=' + '|'.join(supported_roles)
+            args.append('-VPRole=' + '|'.join(supported_roles))
 
         if unsupported_roles:
             LOGGER.error(
@@ -2413,30 +2440,25 @@ class DeviceUnreal(Device):
         # Session ID
         session_id = DeviceUnreal.csettings["stage_session_id"].get_value()
         if session_id > 0:
-            command_line_args += f" -StageSessionId={session_id}"
+            args.append(f'-StageSessionId={session_id}')
 
-        command_line_args += f' -StageFriendlyName="{self.name.replace(" ", "_")}"'
+        # Stage friendly name
+        args.append(f'-StageFriendlyName="{self.name.replace(" ", "_")}"')
 
         # Max GPU Count (mGPU)
-        max_gpu_count = DeviceUnreal.csettings["max_gpu_count"].get_value(
-            self.name)
+        max_gpu_count = DeviceUnreal.csettings["max_gpu_count"].get_value(self.name)
         try:
             if int(max_gpu_count) > 1:
-                command_line_args += f" -MaxGPUCount={max_gpu_count} "
+                args.append(f'-MaxGPUCount={max_gpu_count}')
                 dp_cvars.append('r.AllowMultiGPUInEditor=1')
         except ValueError:
             LOGGER.warning(f"Invalid Number of GPUs '{max_gpu_count}'")
 
-        # Slate.bAllowThrottling. Makes ICVFX panel and Vcam more responsive to Editor interactive changes.
+        # Slate.bAllowThrottling - Makes ICVFX panel and Vcam more responsive
         slate_allow_throttling = DeviceUnreal.csettings["slate_allow_throttling"].get_value()
         dp_cvars.append(f'Slate.bAllowThrottling={int(slate_allow_throttling)}')
 
-        # mediaprofile
-        mediaprofile_gamepath = DeviceUnreal.csettings["mediaprofile"].get_value(self.name)
-        if mediaprofile_gamepath:
-            dp_cvars.append(self.dpcvar_for_mediaprofile(mediaprofile_gamepath))
-
-        # Holdout composite on Editors if this is a Switchboard config with nDisplay devices.
+        # Holdout composite on Editors if this is a Switchboard config with nDisplay devices
         if DeviceUnreal.csettings["force_holdout_if_ndisplay"].get_value(self.name):
             if any(device.device_type == "nDisplay" for device in self.active_unreal_devices):
                 dp_cvars.extend([
@@ -2444,63 +2466,66 @@ class DeviceUnreal(Device):
                     'r.Deferred.SupportPrimitiveAlphaHoldout=1'
                 ])
 
-        # Add user set dp cvars, overriding any of the forced ones.
+        # Add user-set dp cvars, overriding any of the forced ones
         user_dp_cvars = self.conform_dpcvars(DeviceUnreal.csettings["dp_cvars"].get_value(self.name))
         dp_cvars = self.add_or_override_cvars(dp_cvars, user_dp_cvars)
 
-        # add accumulated dpcvars to args
+        # Add accumulated DPCVars to args
         if len(dp_cvars):
-            command_line_args += f" -DPCVars=\"{','.join(dp_cvars)}\""
+            args.append(f'-DPCVars="{",".join(dp_cvars)}"')
 
-        if DeviceUnreal.csettings['auto_decline_package_recovery'].get_value(
-                self.name):
-            command_line_args += ' -AutoDeclinePackageRecovery'
+        # Package recovery
+        if DeviceUnreal.csettings['auto_decline_package_recovery'].get_value(self.name):
+            args.append('-AutoDeclinePackageRecovery')
 
         # UdpMessaging endpoints
         if self.udpmessaging_multicast_endpoint:
             server_mc = CONFIG.MUSERVER_MULTICAST_ENDPOINT.get_value().strip()
             if server_mc != self.udpmessaging_multicast_endpoint:
                 LOGGER.warning(f"{self.name} contains a multicast endpoint ('{self.udpmessaging_multicast_endpoint}') that does not match configured value on Multi-user server ('{server_mc}').")
-            command_line_args += (
-                ' -UDPMESSAGING_TRANSPORT_MULTICAST='
-                f'"{self.udpmessaging_multicast_endpoint}"'
-            )
+            args.append(f'-UDPMESSAGING_TRANSPORT_MULTICAST="{self.udpmessaging_multicast_endpoint}"')
 
         if self.udpmessaging_unicast_endpoint:
-            command_line_args += (
-                ' -UDPMESSAGING_TRANSPORT_UNICAST='
-                f'"{self.udpmessaging_unicast_endpoint}"')
+            args.append(f'-UDPMESSAGING_TRANSPORT_UNICAST="{self.udpmessaging_unicast_endpoint}"')
 
         static_endpoints = self.build_udpmessaging_static_endpoint_list()
         if len(static_endpoints) > 0:
-            command_line_args += (
-                ' -UDPMESSAGING_TRANSPORT_STATIC='
-                f'"{",".join(static_endpoints)}"')
+            args.append(f'-UDPMESSAGING_TRANSPORT_STATIC="{",".join(static_endpoints)}"')
 
         # Disable ensures
         if DeviceUnreal.csettings['disable_ensures'].get_value(self.name):
-            command_line_args += ' -handleensurepercent=0'
+            args.append('-handleensurepercent=0')
 
+        # Concert take recorder settings
         record_on_client = 'True' if self.is_recording_device else 'False'
         ini_engine = (
-            ' -ini:Engine:'
+            '-ini:Engine:'
             '[/Script/ConcertTakeRecorder.ConcertSessionRecordSettings]:'
             f'LocalSettings=(bRecordOnClient={record_on_client})'
         )
-        command_line_args += ini_engine
-
-        command_line_args += f' -VPOSCServerAddress={self.address}'
-        command_line_args += f' -VPOSCServerPort={self.device_osc_port}'
+        args.append(ini_engine)
+
+        # VP OSC server settings
+        args.extend([
+            f'-VPOSCServerAddress={self.address}',
+            f'-VPOSCServerPort={self.device_osc_port}'
+        ])
+
+        # Disable auto-save.
+        args.append(
+            "-ini:EditorPerProjectUserSettings:"
+            "[/Script/UnrealEd.EditorLoadingSavingSettings]:"
+            "bAutoSaveEnable=False"
+        )
 
-        # Logging
+        # Logging categories
         logging_args = DeviceUnreal.csettings['logging'].get_command_line_arg(
             override_device_name=self.name)
         if logging_args:
-            command_line_args += f' {logging_args}'
+            args.append(logging_args)
 
-        return (
-            f'"{CONFIG.UPROJECT_PATH.get_value(self.name)}" {map_name} '
-            f'{command_line_args}')
+        # Join all arguments with spaces
+        return ' '.join(args)
 
     def generate_unreal_command_line(self, map_name):
         return (
