Switchboard: Pass UdpMessaging command line arguments to Live Link Hub.

diff --git a/Engine/Plugins/VirtualProduction/Switchboard/Source/Switchboard/switchboard/devices/livelinkhub/plugin_livelinkhub.py b/Engine/Plugins/VirtualProduction/Switchboard/Source/Switchboard/switchboard/devices/livelinkhub/plugin_livelinkhub.py
--- a/Engine/Plugins/VirtualProduction/Switchboard/Source/Switchboard/switchboard/devices/livelinkhub/plugin_livelinkhub.py
+++ b/Engine/Plugins/VirtualProduction/Switchboard/Source/Switchboard/switchboard/devices/livelinkhub/plugin_livelinkhub.py
@@ -94,6 +94,8 @@ class DeviceLiveLinkHub(DeviceUnreal):
     #@override
     def setting_overrides(self):
         return [
+            DeviceUnreal.csettings['udpmessaging_unicast_endpoint'],
+            DeviceUnreal.csettings['udpmessaging_extra_static_endpoints'],
             DeviceLiveLinkHub.csettings['command_line_arguments'],
             DeviceLiveLinkHub.csettings['session_path'],
             CONFIG.ENGINE_DIR,
@@ -142,27 +144,31 @@ class DeviceLiveLinkHub(DeviceUnreal):
 
     #@override
     def generate_unreal_command_line_args(self, map_name: str):
-        command_line_args = f'{self.extra_cmdline_args_setting}'
+        command_line_args: list[str] = []
+        command_line_args.append(f'{self.extra_cmdline_args_setting}')
 
-        command_line_args += f' Log={self.log_filename}'
+        command_line_args.append(f'Log={self.log_filename}')
 
         if self.setting_spoke_mode.get_value() == False:
-            command_line_args += ' -hub'
-            command_line_args += ' -UDPMESSAGING_SHARE_KNOWN_NODES=1'
+            command_line_args.extend(['-hub',
+                                      '-UDPMESSAGING_SHARE_KNOWN_NODES=1'])
         else:
-            command_line_args += ' -spoke'
+            command_line_args.append('-spoke')
+
+        # UdpMessaging endpoints
+        command_line_args.extend(self.build_udpmessaging_args())
 
         if self.session_path:
-            command_line_args += f' -SessionPath="{self.session_path}"'
+            command_line_args.append(f' -SessionPath="{self.session_path}"')
 
         if DeviceLiveLinkHub.csettings['disable_crash_recovery'].get_value():
-            command_line_args += (
+            command_line_args.append(
                 ' -ini:Engine:'
                 '[/Script/LiveLinkHub.LiveLinkHubSettings]:'
                 'bEnableCrashRecovery=false'
             )
 
-        return command_line_args
+        return ' '.join(command_line_args)
 
     #@override
     def launch(self, map_name: str):
diff --git a/Engine/Plugins/VirtualProduction/Switchboard/Source/Switchboard/switchboard/devices/unreal/plugin_unreal.py b/Engine/Plugins/VirtualProduction/Switchboard/Source/Switchboard/switchboard/devices/unreal/plugin_unreal.py
--- a/Engine/Plugins/VirtualProduction/Switchboard/Source/Switchboard/switchboard/devices/unreal/plugin_unreal.py
+++ b/Engine/Plugins/VirtualProduction/Switchboard/Source/Switchboard/switchboard/devices/unreal/plugin_unreal.py
@@ -2291,6 +2291,23 @@ class DeviceUnreal(Device):
 
         return endpoints
 
+    def build_udpmessaging_args(self) -> list[str]:
+        args: list[str] = []
+        if self.udpmessaging_multicast_endpoint:
+            server_mc = CONFIG.MUSERVER_MULTICAST_ENDPOINT.get_value().strip()
+            if server_mc != self.udpmessaging_multicast_endpoint:
+                LOGGER.warning(f"{self.name} contains a multicast endpoint ('{self.udpmessaging_multicast_endpoint}') that does not match configured value on Multi-user server ('{server_mc}').")
+            args.append(f'-UDPMESSAGING_TRANSPORT_MULTICAST="{self.udpmessaging_multicast_endpoint}"')
+
+        if self.udpmessaging_unicast_endpoint:
+            args.append(f'-UDPMESSAGING_TRANSPORT_UNICAST="{self.udpmessaging_unicast_endpoint}"')
+
+        static_endpoints = self.build_udpmessaging_static_endpoint_list()
+        if len(static_endpoints) > 0:
+            args.append(f'-UDPMESSAGING_TRANSPORT_STATIC="{",".join(static_endpoints)}"')
+
+        return args
+
     def get_utrace_filepath(self):
         return self.get_remote_log_path() / f'{self.name}_{self.runtime_str}.utrace'
 
@@ -2364,7 +2381,7 @@ class DeviceUnreal(Device):
     def generate_unreal_command_line_args(self, map_name):
 
         # Build command line arguments as an array that will be joined at the end.
-        args = []
+        args: list[str] = []
 
         # DPCVars - accumulate these before adding to args
         dp_cvars = []
@@ -2479,18 +2496,7 @@ class DeviceUnreal(Device):
             args.append('-AutoDeclinePackageRecovery')
 
         # UdpMessaging endpoints
-        if self.udpmessaging_multicast_endpoint:
-            server_mc = CONFIG.MUSERVER_MULTICAST_ENDPOINT.get_value().strip()
-            if server_mc != self.udpmessaging_multicast_endpoint:
-                LOGGER.warning(f"{self.name} contains a multicast endpoint ('{self.udpmessaging_multicast_endpoint}') that does not match configured value on Multi-user server ('{server_mc}').")
-            args.append(f'-UDPMESSAGING_TRANSPORT_MULTICAST="{self.udpmessaging_multicast_endpoint}"')
-
-        if self.udpmessaging_unicast_endpoint:
-            args.append(f'-UDPMESSAGING_TRANSPORT_UNICAST="{self.udpmessaging_unicast_endpoint}"')
-
-        static_endpoints = self.build_udpmessaging_static_endpoint_list()
-        if len(static_endpoints) > 0:
-            args.append(f'-UDPMESSAGING_TRANSPORT_STATIC="{",".join(static_endpoints)}"')
+        args.extend(self.build_udpmessaging_args())
 
         # Disable ensures
         if DeviceUnreal.csettings['disable_ensures'].get_value(self.name):
