Switchboard: Only consider devices with a valid osc_client to be eligible to handle incoming messages.

This change also ensures the LiveLinkHub device will never create a valid osc_client.

diff --git a/Engine/Plugins/VirtualProduction/Switchboard/Source/Switchboard/switchboard/devices/livelinkhub/plugin_livelinkhub.py b/Engine/Plugins/VirtualProduction/Switchboard/Source/Switchboard/switchboard/devices/livelinkhub/plugin_livelinkhub.py
--- a/Engine/Plugins/VirtualProduction/Switchboard/Source/Switchboard/switchboard/devices/livelinkhub/plugin_livelinkhub.py
+++ b/Engine/Plugins/VirtualProduction/Switchboard/Source/Switchboard/switchboard/devices/livelinkhub/plugin_livelinkhub.py
@@ -138,6 +138,16 @@ class DeviceLiveLinkHub(DeviceUnreal):
     def session_path(self):
         return DeviceLiveLinkHub.csettings['session_path'].get_value(self.name)
 
+    #@override
+    def setup_osc_client(self, osc_port):
+        ''' Live Link Hub devices should not send or receive OSC '''
+        return
+
+    #@override
+    def send_osc_message(self, command, value, log=True):
+        ''' Live Link Hub devices should not send or receive OSC '''
+        return
+
     #@override
     def get_remote_log_path(self):
         return self.llh_proj_dir / 'Saved' / 'Logs'
diff --git a/Engine/Plugins/VirtualProduction/Switchboard/Source/Switchboard/switchboard/switchboard_dialog.py b/Engine/Plugins/VirtualProduction/Switchboard/Source/Switchboard/switchboard/switchboard_dialog.py
--- a/Engine/Plugins/VirtualProduction/Switchboard/Source/Switchboard/switchboard/switchboard_dialog.py
+++ b/Engine/Plugins/VirtualProduction/Switchboard/Source/Switchboard/switchboard/switchboard_dialog.py
@@ -2479,7 +2479,7 @@ class SwitchboardDialog(QtCore.QObject):
             self.p4_refresh_project_cl()
 
     def osc_take(self, address, command, value):
-        device = self._device_from_address(address, command, value=value)
+        device = self._osc_device_from_address(address, command, value=value)
 
         if not device:
             return
@@ -2488,7 +2488,7 @@ class SwitchboardDialog(QtCore.QObject):
 
     # OSC: Set Slate
     def osc_slate(self, address, command, value):
-        device = self._device_from_address(address, command, value=value)
+        device = self._osc_device_from_address(address, command, value=value)
         if not device:
             return
 
@@ -2510,11 +2510,16 @@ class SwitchboardDialog(QtCore.QObject):
             LOGGER.debug('Record start button pressed')
             self._record_start(self.slate, self.take, self.description)
 
-    def _device_from_address(self, address, command, value=''):
-        device = self.device_manager.device_with_address(address[0])
+    def _osc_device_from_address(self, address, command, value=''):
+        device: Device | None = None
+
+        for check_device in self.device_manager.devices():
+            if (check_device.address == address[0]) and (check_device.osc_client is not None):
+                device = check_device
+                break
 
         if not device:
-            LOGGER.warning(f'{address} is not registered with a device in Switchboard')
+            LOGGER.warning(f'{address} is not registered with an OSC device in Switchboard')
             return None
 
         # Do not receive osc from disconnected devices
@@ -2530,7 +2535,7 @@ class SwitchboardDialog(QtCore.QObject):
         '''
         OSC message Recieved /RecordStart
         '''
-        device = self._device_from_address(address, command, value=[slate, take, description])
+        device = self._osc_device_from_address(address, command, value=[slate, take, description])
         if not device:
             return
 
@@ -2633,21 +2638,21 @@ class SwitchboardDialog(QtCore.QObject):
         #self.take = new_recording.take + 1
 
     def osc_record_start_confirm(self, address, command, timecode):
-        device = self._device_from_address(address, command, value=timecode)
+        device = self._osc_device_from_address(address, command, value=timecode)
         if not device:
             return
 
         device.record_start_confirm(timecode)
 
     def osc_record_stop(self, address, command):
-        device = self._device_from_address(address, command)
+        device = self._osc_device_from_address(address, command)
         if not device:
             return
 
         self._record_stop(exclude_address=device.address)
 
     def osc_record_stop_confirm(self, address, command, timecode, *paths):
-        device = self._device_from_address(address, command, value=timecode)
+        device = self._osc_device_from_address(address, command, value=timecode)
         if not device:
             return
 
@@ -2660,7 +2665,7 @@ class SwitchboardDialog(QtCore.QObject):
         This is called when record has been pressed and stopped before the countdown in take recorder
         has finished
         """
-        device = self._device_from_address(address, command)
+        device = self._osc_device_from_address(address, command)
         if not device:
             return
 
@@ -2668,14 +2673,14 @@ class SwitchboardDialog(QtCore.QObject):
 
     def osc_record_cancel_confirm(self, address, command, timecode):
         pass
-        #device = self._device_from_address(address, command, value=timecode)
+        #device = self._osc_device_from_address(address, command, value=timecode)
         #if not device:
         #    return
 
         #self.record_cancel_confirm(device, timecode)
 
     def osc_ue4_launch_confirm(self, address, command):
-        device = self._device_from_address(address, command)
+        device = self._osc_device_from_address(address, command)
         if not device:
             return
 
@@ -2697,7 +2702,7 @@ class SwitchboardDialog(QtCore.QObject):
         LOGGER.debug(f'osc_arsession_stop_confirm {value}')
 
     def osc_arsession_start_confirm(self, address, command, value):
-        device = self._device_from_address(address, command, value=value)
+        device = self._osc_device_from_address(address, command, value=value)
         if not device:
             return
 
@@ -2718,7 +2723,7 @@ class SwitchboardDialog(QtCore.QObject):
         device_widget.set_battery(value)
 
     def osc_data(self, address, command, value):
-        device = self._device_from_address(address, command)
+        device = self._osc_device_from_address(address, command)
         if not device:
             return
 
