Switchboard: Allow hyphens in asset names when parsing references.

Cherrypick of CL 46029014.

#rnx

#Synced-CL 48062107

diff --git a/Engine/Plugins/VirtualProduction/Switchboard/Source/Switchboard/switchboard/devices/unreal/plugin_unreal.py b/Engine/Plugins/VirtualProduction/Switchboard/Source/Switchboard/switchboard/devices/unreal/plugin_unreal.py
--- a/Engine/Plugins/VirtualProduction/Switchboard/Source/Switchboard/switchboard/devices/unreal/plugin_unreal.py
+++ b/Engine/Plugins/VirtualProduction/Switchboard/Source/Switchboard/switchboard/devices/unreal/plugin_unreal.py
@@ -3697,9 +3697,9 @@ class DeviceUnreal(Device):
             \.                                     # .
             (?P<classname>[A-Za-z_]\w+)            # Class
             '                                      # '
-            (?P<gamepath>[A-Za-z0-9_/]+)           # Gamepath
+            (?P<gamepath>[A-Za-z0-9_/\-]+)         # Gamepath (allow -)
             \.                                     # .
-            (?P<name>[A-Za-z0-9_]+)                # Name
+            (?P<name>[A-Za-z0-9_\-]+)              # Name (allow -)
             '                                      # '
             $""", re.VERBOSE)
 
