Switchboard: Added a way to set the media profile and live link preset at runtime.

Includes a way to clear them (select "None" and click "Set").

This was implemented via console commands, and relies on nDisplay cluster events. The media profile / LL preset are applied to the whole cluster, it isn't yet possible to target a specific set of nodes.

The needed console commands were added to Unreal.

(Cherrypick of CL 46045817.)

#Synced-CL 48062107

diff --git a/Engine/Plugins/Animation/LiveLink/Source/LiveLink/Private/LiveLinkPreset.cpp b/Engine/Plugins/Animation/LiveLink/Source/LiveLink/Private/LiveLinkPreset.cpp
--- a/Engine/Plugins/Animation/LiveLink/Source/LiveLink/Private/LiveLinkPreset.cpp
+++ b/Engine/Plugins/Animation/LiveLink/Source/LiveLink/Private/LiveLinkPreset.cpp
@@ -61,6 +61,25 @@ namespace
 				}
 			})
 		);
+
+	static FAutoConsoleCommand LiveLinkPresetClearCmd(
+		TEXT("LiveLink.Preset.Clear"),
+		TEXT("Remove all LiveLink sources from the client"),
+		FConsoleCommandDelegate::CreateLambda([]()
+			{
+				if (IModularFeatures::Get().IsModularFeatureAvailable(ILiveLinkClient::ModularFeatureName))
+				{
+					FLiveLinkClient& LiveLinkClient = IModularFeatures::Get().GetModularFeature<FLiveLinkClient>(ILiveLinkClient::ModularFeatureName);
+					LiveLinkClient.RemoveAllSources();
+					LiveLinkClient.Tick();
+					FLiveLinkLog::Info(TEXT("Removed all LiveLink sources"));
+				}
+				else
+				{
+					FLiveLinkLog::Error(TEXT("LiveLink client not available"));
+				}
+			})
+		);
 }
 
 struct FApplyToClientPollingOperation
diff --git a/Engine/Plugins/Media/MediaFrameworkUtilities/Source/MediaFrameworkUtilities/Private/MediaFrameworkUtilitiesModule.cpp b/Engine/Plugins/Media/MediaFrameworkUtilities/Source/MediaFrameworkUtilities/Private/MediaFrameworkUtilitiesModule.cpp
--- a/Engine/Plugins/Media/MediaFrameworkUtilities/Source/MediaFrameworkUtilities/Private/MediaFrameworkUtilitiesModule.cpp
+++ b/Engine/Plugins/Media/MediaFrameworkUtilities/Source/MediaFrameworkUtilities/Private/MediaFrameworkUtilitiesModule.cpp
@@ -26,6 +26,59 @@ static TAutoConsoleVariable<FString> CVarMediaUtilsStartupProfile(
 	ECVF_ReadOnly
 );
 
+static FAutoConsoleCommand MediaProfileSetCmd(
+	TEXT("MediaProfile.Set"),
+	TEXT("Set the current Media Profile. Usage: MediaProfile.Set /Game/Path/To/MediaProfile.MediaProfile (or no argument to clear)"),
+	FConsoleCommandWithArgsDelegate::CreateLambda([](const TArray<FString>& Args)
+	{
+		if (Args.Num() == 0)
+		{
+			// Clear the media profile when called with no arguments
+			if (IMediaFrameworkUtilitiesModule* Module = FModuleManager::GetModulePtr<IMediaFrameworkUtilitiesModule>("MediaFrameworkUtilities"))
+			{
+				Module->GetProfileManager().SetCurrentMediaProfile(nullptr);
+				UE_LOG(LogMediaFrameworkUtilities, Display, TEXT("Cleared Media Profile"));
+			}
+			return;
+		}
+
+		const FString& ProfilePath = Args[0];
+
+		if (UObject* Object = StaticLoadObject(UMediaProfile::StaticClass(), nullptr, *ProfilePath))
+		{
+			if (UMediaProfile* MediaProfile = Cast<UMediaProfile>(Object))
+			{
+				if (IMediaFrameworkUtilitiesModule* Module = FModuleManager::GetModulePtr<IMediaFrameworkUtilitiesModule>("MediaFrameworkUtilities"))
+				{
+					Module->GetProfileManager().SetCurrentMediaProfile(MediaProfile);
+					UE_LOG(LogMediaFrameworkUtilities, Display, TEXT("Applied Media Profile: %s"), *ProfilePath);
+				}
+			}
+			else
+			{
+				UE_LOG(LogMediaFrameworkUtilities, Error, TEXT("Object at path '%s' is not a Media Profile"), *ProfilePath);
+			}
+		}
+		else
+		{
+			UE_LOG(LogMediaFrameworkUtilities, Error, TEXT("Failed to load Media Profile from path: %s"), *ProfilePath);
+		}
+	})
+);
+
+static FAutoConsoleCommand MediaProfileClearCmd(
+	TEXT("MediaProfile.Clear"),
+	TEXT("Clear the current Media Profile"),
+	FConsoleCommandDelegate::CreateLambda([]()
+	{
+		if (IMediaFrameworkUtilitiesModule* Module = FModuleManager::GetModulePtr<IMediaFrameworkUtilitiesModule>("MediaFrameworkUtilities"))
+		{
+			Module->GetProfileManager().SetCurrentMediaProfile(nullptr);
+			UE_LOG(LogMediaFrameworkUtilities, Display, TEXT("Cleared Media Profile"));
+		}
+	})
+);
+
 /**
  * Implements the MediaFrameworkUtilitiesModule module.
  */
diff --git a/Engine/Plugins/Media/MediaFrameworkUtilities/Source/MediaFrameworkUtilities/Private/Profile/MediaProfile.cpp b/Engine/Plugins/Media/MediaFrameworkUtilities/Source/MediaFrameworkUtilities/Private/Profile/MediaProfile.cpp
--- a/Engine/Plugins/Media/MediaFrameworkUtilities/Source/MediaFrameworkUtilities/Private/Profile/MediaProfile.cpp
+++ b/Engine/Plugins/Media/MediaFrameworkUtilities/Source/MediaFrameworkUtilities/Private/Profile/MediaProfile.cpp
@@ -523,62 +523,166 @@ void UMediaProfile::Reset()
 
 void UMediaProfile::ResetTimecodeProvider()
 {
-	if (bTimecodeProvideWasApplied)
+	if (!bTimecodeProvideWasApplied)
 	{
-		if (AppliedTimecodeProvider == GEngine->GetTimecodeProvider())
+		return;
+	}
+
+	// Optional but recommended in UE code that touches engine-wide state
+	ensure(IsInGameThread());
+
+	if (!GEngine)
+	{
+		UE_LOG(LogMediaFrameworkUtilities, Warning, TEXT("Cannot reset TimecodeProvider: GEngine is null."));
+		return; // keep state so caller can retry later
+	}
+
+	auto GetNameSafe = [](const UObject* Obj) -> FString
 		{
-			bool bResult = GEngine->SetTimecodeProvider(PreviousTimecodeProvider);
-			if (!bResult && PreviousTimecodeProvider)
-			{
-				UE_LOG(LogMediaFrameworkUtilities, Warning, TEXT("The TimecodeProvider '%s' could not be initialized."), *PreviousTimecodeProvider->GetName());
-			}
+			return (Obj && IsValid(Obj)) ? Obj->GetName() : TEXT("None");
+		};
+
+	UTimecodeProvider* const CurrentProvider = GEngine->GetTimecodeProvider();
+
+	// Only restore if we still own the setting (don’t clobber external changes)
+	const bool bAppliedValid = (AppliedTimecodeProvider && IsValid(AppliedTimecodeProvider));
+	const bool bWeStillOwn = (bAppliedValid && AppliedTimecodeProvider == CurrentProvider);
+
+	bool bRestored = false;
+
+	if (bWeStillOwn)
+	{
+		// If previous is invalid, reset to null before attempting to apply
+		if (PreviousTimecodeProvider && !IsValid(PreviousTimecodeProvider))
+		{
+			UE_LOG(LogMediaFrameworkUtilities, Warning, TEXT("Previous TimecodeProvider is no longer valid. Resetting to null."));
+			PreviousTimecodeProvider = nullptr;
 		}
-		else
+
+		bRestored = GEngine->SetTimecodeProvider(PreviousTimecodeProvider);
+
+		if (!bRestored)
 		{
-			if (PreviousTimecodeProvider)
+			if (PreviousTimecodeProvider && IsValid(PreviousTimecodeProvider))
 			{
-				UE_LOG(LogMediaFrameworkUtilities, Warning, TEXT("Could not set the previous TimecodeProvider '%s'."), *PreviousTimecodeProvider->GetName());
+				UE_LOG(LogMediaFrameworkUtilities, Warning,
+					TEXT("The TimecodeProvider '%s' could not be initialized."),
+					*PreviousTimecodeProvider->GetName());
 			}
 			else
 			{
-				UE_LOG(LogMediaFrameworkUtilities, Warning, TEXT("Could not set the previous TimecodeProvider."));
+				UE_LOG(LogMediaFrameworkUtilities, Warning,
+					TEXT("Could not clear the TimecodeProvider."));
 			}
 		}
+	}
+	else
+	{
+		// Clarify whether it changed externally or our applied instance is gone
+		if (!bAppliedValid && AppliedTimecodeProvider)
+		{
+			UE_LOG(LogMediaFrameworkUtilities, Warning,
+				TEXT("Previously applied TimecodeProvider is no longer valid. Not restoring."));
+		}
+		else
+		{
+			UE_LOG(LogMediaFrameworkUtilities, Warning,
+				TEXT("TimecodeProvider changed externally (current: '%s', previously applied: '%s'). Not restoring."),
+				*GetNameSafe(CurrentProvider), *GetNameSafe(AppliedTimecodeProvider));
+		}
+	}
+
+	// Only clear state if we either restored successfully or we explicitly decided not to restore
+	if (bRestored || !bWeStillOwn)
+	{
 		PreviousTimecodeProvider = nullptr;
 		AppliedTimecodeProvider = nullptr;
 		bTimecodeProvideWasApplied = false;
 	}
 }
 
+
 void UMediaProfile::ResetCustomTimeStep()
 {
-	if (bCustomTimeStepWasApplied)
+	if (!bCustomTimeStepWasApplied)
 	{
-		if (AppliedCustomTimeStep == GEngine->GetCustomTimeStep())
+		return;
+	}
+
+	// Recommended: enforce game-thread usage when touching engine-wide state
+	ensure(IsInGameThread());
+
+	if (!GEngine)
+	{
+		UE_LOG(LogMediaFrameworkUtilities, Warning, TEXT("Cannot reset CustomTimeStep: GEngine is null."));
+		return; // keep state so we can retry later
+	}
+
+	auto GetNameSafe = [](const UObject* Obj) -> FString
 		{
-			bool bResult = GEngine->SetCustomTimeStep(PreviousCustomTimeStep);
-			if (!bResult && PreviousCustomTimeStep)
-			{
-				UE_LOG(LogMediaFrameworkUtilities, Warning, TEXT("The Custom Time Step '%s' could not be initialized."), *PreviousCustomTimeStep->GetName());
-			}
+			return (Obj && IsValid(Obj)) ? Obj->GetName() : TEXT("None");
+		};
+
+	// Ownership check: only restore if the engine still has what we applied
+	UEngineCustomTimeStep* const CurrentCTS = GEngine->GetCustomTimeStep();
+	const bool bAppliedValid = (AppliedCustomTimeStep && IsValid(AppliedCustomTimeStep));
+	const bool bWeStillOwn = (bAppliedValid && AppliedCustomTimeStep == CurrentCTS);
+
+	bool bRestored = false;
+
+	if (bWeStillOwn)
+	{
+		// If previous is invalid, reset to null to avoid passing a dead object
+		if (PreviousCustomTimeStep && !IsValid(PreviousCustomTimeStep))
+		{
+			UE_LOG(LogMediaFrameworkUtilities, Warning,
+				TEXT("Previous Custom Time Step is no longer valid. Resetting to null."));
+			PreviousCustomTimeStep = nullptr;
 		}
-		else
+
+		bRestored = GEngine->SetCustomTimeStep(PreviousCustomTimeStep);
+
+		if (!bRestored)
 		{
-			if (PreviousCustomTimeStep)
+			if (PreviousCustomTimeStep && IsValid(PreviousCustomTimeStep))
 			{
-				UE_LOG(LogMediaFrameworkUtilities, Warning, TEXT("Could not set the previous Custom Time Step '%s'."), *PreviousCustomTimeStep->GetName());
+				UE_LOG(LogMediaFrameworkUtilities, Warning,
+					TEXT("The Custom Time Step '%s' could not be initialized."),
+					*PreviousCustomTimeStep->GetName());
 			}
 			else
 			{
-				UE_LOG(LogMediaFrameworkUtilities, Warning, TEXT("Could not set the previous Custom Time Step."));
+				UE_LOG(LogMediaFrameworkUtilities, Warning,
+					TEXT("Could not clear the Custom Time Step."));
 			}
 		}
+	}
+	else
+	{
+		// Clarify why we’re not restoring
+		if (!bAppliedValid && AppliedCustomTimeStep)
+		{
+			UE_LOG(LogMediaFrameworkUtilities, Warning,
+				TEXT("Previously applied Custom Time Step is no longer valid. Not restoring."));
+		}
+		else
+		{
+			UE_LOG(LogMediaFrameworkUtilities, Warning,
+				TEXT("Custom Time Step changed externally (current: '%s', previously applied: '%s'). Not restoring."),
+				*GetNameSafe(CurrentCTS), *GetNameSafe(AppliedCustomTimeStep));
+		}
+	}
+
+	// Only clear our local state if we either restored successfully or explicitly chose not to restore
+	if (bRestored || !bWeStillOwn)
+	{
 		PreviousCustomTimeStep = nullptr;
 		AppliedCustomTimeStep = nullptr;
 		bCustomTimeStepWasApplied = false;
 	}
 }
 
+
 /**
  * @EventName MediaFramework.ApplyMediaProfile
  * @Trigger Triggered when a media profile is applied.
diff --git a/Engine/Plugins/VirtualProduction/Switchboard/Source/Switchboard/switchboard/devices/ndisplay/ndisplay_tools_ui.py b/Engine/Plugins/VirtualProduction/Switchboard/Source/Switchboard/switchboard/devices/ndisplay/ndisplay_tools_ui.py
new file mode 100644
--- /dev/null
+++ b/Engine/Plugins/VirtualProduction/Switchboard/Source/Switchboard/switchboard/devices/ndisplay/ndisplay_tools_ui.py
@@ -0,0 +1,339 @@
+# Copyright Epic Games, Inc. All Rights Reserved.
+
+from typing import Optional
+
+from PySide6 import QtCore, QtGui, QtWidgets
+from PySide6.QtWidgets import (
+    QComboBox, QHBoxLayout, QLabel, QPushButton,
+    QSizePolicy, QVBoxLayout, QWidget, QGroupBox
+)
+
+from switchboard import switchboard_widgets as sb_widgets
+from switchboard.config import CONFIG
+from switchboard.sbcache import SBCache, Asset
+from switchboard.switchboard_logging import LOGGER
+from switchboard.devices.unreal.plugin_unreal import DeviceUnreal
+
+
+class nDisplayToolsUI(QWidget):
+    '''
+    UI for nDisplay Tools tab that allows setting Media Profiles and
+    applying Live Link presets to the nDisplay cluster at runtime.
+    '''
+
+    def __init__(self, parent: QWidget, plugin_ndisplay):
+        QWidget.__init__(self, parent)
+
+        self.plugin_ndisplay = plugin_ndisplay
+
+        # Main layout
+        layout = QVBoxLayout(self)
+        layout.setSpacing(10)
+
+        # Media Profile section
+        media_profile_group = self.create_media_profile_section()
+        layout.addWidget(media_profile_group)
+
+        # Live Link Preset section
+        livelink_preset_group = self.create_livelink_preset_section()
+        layout.addWidget(livelink_preset_group)
+
+        # Add stretch to push everything to the top
+        layout.addStretch()
+
+        # Listen for clipboard changes to auto-refresh
+        QtWidgets.QApplication.clipboard().dataChanged.connect(
+            self.on_clipboard_changed)
+
+        # Listen for asset additions from clipboard operations
+        DeviceUnreal.static_signals.clipboard_assets_added_signal.connect(
+            self.on_assets_added_via_clipboard)
+
+        # Listen for project asset scans/updates
+        DeviceUnreal.static_signals.project_assets_updated_signal.connect(
+            self.on_project_assets_updated)
+
+        # Listen for nDisplay config file changes
+        self.plugin_ndisplay.csettings['ndisplay_config_file'].signal_setting_changed.connect(
+            self.on_config_changed)
+
+        # Listen for default media profile/livelink preset changes
+        self.plugin_ndisplay.csettings['mediaprofile'].signal_setting_changed.connect(
+            self.on_default_settings_changed)
+        self.plugin_ndisplay.csettings['livelink_preset'].signal_setting_changed.connect(
+            self.on_default_settings_changed)
+
+        # Initial population of combo boxes
+        self.refresh_media_profiles()
+        self.refresh_livelink_presets()
+
+    def scan_project_assets(self):
+        '''Scans project for assets and updates cache, following UAssetSetting pattern'''
+        DeviceUnreal.analyze_project_assets()
+
+    def on_project_assets_updated(self):
+        '''Called when project assets are updated from any source'''
+        self.refresh_media_profiles()
+        self.refresh_livelink_presets()
+
+    def on_config_changed(self, old_value, new_value):
+        '''Called when nDisplay config file changes - update default selections'''
+        # When config changes, the default media profile and livelink preset may have changed
+        # So we need to update our dropdown selections
+        self._set_default_media_profile()
+        self._set_default_livelink_preset()
+
+    def on_default_settings_changed(self, old_value, new_value):
+        '''Called when default media profile or livelink preset settings change'''
+        # Update dropdown selections to reflect new defaults
+        self._set_default_media_profile()
+        self._set_default_livelink_preset()
+
+    def _set_default_media_profile(self):
+        '''Sets the media profile combo box to the default value from nDisplay device settings'''
+        # Get default media profile from nDisplay class settings
+        default_gamepath = self.plugin_ndisplay.csettings['mediaprofile'].get_value()
+
+        if default_gamepath:
+            # Find the item in combo box that matches this gamepath
+            for i in range(self.media_profile_combo.count()):
+                if self.media_profile_combo.itemData(i) == default_gamepath:
+                    self.media_profile_combo.setCurrentIndex(i)
+                    break
+
+    def _set_default_livelink_preset(self):
+        '''Sets the livelink preset combo box to the default value from nDisplay device settings'''
+        # Get default livelink preset from nDisplay class settings
+        default_gamepath = self.plugin_ndisplay.csettings['livelink_preset'].get_value()
+
+        if default_gamepath:
+            # Find the item in combo box that matches this gamepath
+            for i in range(self.livelink_preset_combo.count()):
+                if self.livelink_preset_combo.itemData(i) == default_gamepath:
+                    self.livelink_preset_combo.setCurrentIndex(i)
+                    break
+
+    def create_media_profile_section(self) -> QGroupBox:
+        '''Creates the Media Profile section of the UI'''
+        group_box = QGroupBox("Media Profile")
+        layout = QVBoxLayout()
+
+        # Combo box and buttons layout
+        controls_layout = QHBoxLayout()
+
+        # Media Profile combo box
+        self.media_profile_combo = sb_widgets.SearchableComboBox(self)
+        self.media_profile_combo.setEditable(True)
+        self.media_profile_combo.setSizePolicy(
+            QSizePolicy.Expanding, QSizePolicy.Preferred)
+        self.media_profile_combo.setMinimumWidth(300)
+        self.media_profile_combo.setToolTip(
+            "Select a Media Profile or paste a path from the Content Browser")
+        controls_layout.addWidget(self.media_profile_combo)
+
+        # Refresh button
+        self.media_profile_refresh_btn = QPushButton()
+        self.media_profile_refresh_btn.setIcon(
+            QtGui.QIcon(':/icons/images/icon_refresh.png'))
+        self.media_profile_refresh_btn.setToolTip("Scan project for Media Profiles")
+        self.media_profile_refresh_btn.clicked.connect(
+            self.scan_project_assets)
+        self.media_profile_refresh_btn.setMaximumWidth(30)
+        controls_layout.addWidget(self.media_profile_refresh_btn)
+
+        # Set button
+        self.media_profile_set_btn = QPushButton("Set")
+        self.media_profile_set_btn.setToolTip(
+            "Apply the selected Media Profile to the nDisplay cluster, or clear if (None) is selected")
+        self.media_profile_set_btn.clicked.connect(
+            self.on_set_media_profile)
+        controls_layout.addWidget(self.media_profile_set_btn)
+
+        layout.addLayout(controls_layout)
+        group_box.setLayout(layout)
+
+        return group_box
+
+    def create_livelink_preset_section(self) -> QGroupBox:
+        '''Creates the Live Link Preset section of the UI'''
+        group_box = QGroupBox("Live Link Preset")
+        layout = QVBoxLayout()
+
+        # Combo box and buttons layout
+        controls_layout = QHBoxLayout()
+
+        # Live Link Preset combo box
+        self.livelink_preset_combo = sb_widgets.SearchableComboBox(self)
+        self.livelink_preset_combo.setEditable(True)
+        self.livelink_preset_combo.setSizePolicy(
+            QSizePolicy.Expanding, QSizePolicy.Preferred)
+        self.livelink_preset_combo.setMinimumWidth(300)
+        self.livelink_preset_combo.setToolTip(
+            "Select a Live Link Preset or paste a path from the Content Browser")
+        controls_layout.addWidget(self.livelink_preset_combo)
+
+        # Refresh button
+        self.livelink_preset_refresh_btn = QPushButton()
+        self.livelink_preset_refresh_btn.setIcon(
+            QtGui.QIcon(':/icons/images/icon_refresh.png'))
+        self.livelink_preset_refresh_btn.setToolTip("Scan project for Live Link Presets")
+        self.livelink_preset_refresh_btn.clicked.connect(
+            self.scan_project_assets)
+        self.livelink_preset_refresh_btn.setMaximumWidth(30)
+        controls_layout.addWidget(self.livelink_preset_refresh_btn)
+
+        # Set button
+        self.livelink_preset_set_btn = QPushButton("Set")
+        self.livelink_preset_set_btn.setToolTip(
+            "Apply the selected Live Link Preset to the nDisplay cluster, or remove all sources if (None) is selected")
+        self.livelink_preset_set_btn.clicked.connect(
+            self.on_set_livelink_preset)
+        controls_layout.addWidget(self.livelink_preset_set_btn)
+
+        layout.addLayout(controls_layout)
+        group_box.setLayout(layout)
+
+        return group_box
+
+    def refresh_media_profiles(self):
+        '''Refreshes the Media Profile combo box with assets from the cache'''
+        self.media_profile_combo.clear()
+
+        # Add empty option first
+        self.media_profile_combo.addItem("(None)", "")
+
+        # Query media profiles from cache
+        project = SBCache().query_project(CONFIG.UPROJECT_PATH.get_value())
+        if project:
+            assets = SBCache().query_assets_by_classname(
+                project=project,
+                classnames=DeviceUnreal.MEDIAPROFILE_CLASS_NAMES)
+
+            # Sort assets by name
+            assets = sorted(assets, key=lambda x: x.name)
+
+            # Add to combo box
+            for asset in assets:
+                self.media_profile_combo.addItem(asset.name, asset.gamepath)
+
+        # Set default value from nDisplay device settings
+        self._set_default_media_profile()
+
+    def refresh_livelink_presets(self):
+        '''Refreshes the Live Link Preset combo box with assets from the cache'''
+        self.livelink_preset_combo.clear()
+
+        # Add empty option first
+        self.livelink_preset_combo.addItem("(None)", "")
+
+        # Query livelink presets from cache
+        project = SBCache().query_project(CONFIG.UPROJECT_PATH.get_value())
+        if project:
+            assets = SBCache().query_assets_by_classname(
+                project=project,
+                classnames=DeviceUnreal.LIVELINKPRESET_CLASS_NAMES)
+
+            # Sort assets by name
+            assets = sorted(assets, key=lambda x: x.name)
+
+            # Add to combo box
+            for asset in assets:
+                self.livelink_preset_combo.addItem(asset.name, asset.gamepath)
+
+        # Set default value from nDisplay device settings
+        self._set_default_livelink_preset()
+
+    def on_clipboard_changed(self):
+        '''Called when clipboard content changes - auto-refresh if relevant assets were pasted'''
+        clipboard = QtWidgets.QApplication.clipboard()
+        mime_data = clipboard.mimeData()
+
+        if mime_data.hasText():
+            text = mime_data.text()
+
+            # Check if it looks like an Unreal asset reference
+            if '/Game/' in text or '/Script/' in text:
+                # Refresh both combo boxes since we don't know what was pasted
+                self.refresh_media_profiles()
+                self.refresh_livelink_presets()
+
+    def on_assets_added_via_clipboard(self):
+        '''Called when assets are confirmed to be added via clipboard operations'''
+        # Refresh both combo boxes since new assets were added to the database
+        self.refresh_media_profiles()
+        self.refresh_livelink_presets()
+
+    def get_connected_devices(self):
+        '''Returns a list of connected nDisplay devices'''
+        devices = []
+        if hasattr(self.plugin_ndisplay, 'active_unreal_devices'):
+            for device in self.plugin_ndisplay.active_unreal_devices:
+                if device.device_type == "nDisplay" and device.is_connected_and_authenticated():
+                    devices.append(device)
+        return devices
+
+    def on_set_media_profile(self):
+        '''Sets the selected Media Profile on the cluster or clears it if (None) is selected'''
+        devices = self.get_connected_devices()
+        if not devices:
+            LOGGER.warning("No connected nDisplay devices found")
+            return
+
+        # Get selected media profile gamepath
+        current_index = self.media_profile_combo.currentIndex()
+        if current_index > 0:  # Non-None selection
+            gamepath = self.media_profile_combo.currentData()
+            if gamepath:
+                # Ensure proper formatting of the game path
+                if not gamepath.endswith('.uasset'):
+                    # Extract the asset name and construct the full path
+                    parts = gamepath.split('/')
+                    asset_name = parts[-1] if parts else ''
+                    gamepath = f"{gamepath}.{asset_name}"
+
+                exec_str = f"MediaProfile.Set {gamepath}"
+
+                try:
+                    devices[0].console_exec_cluster(devices, exec_str)
+                    LOGGER.info(f"Set Media Profile: {gamepath}")
+                except Exception as e:
+                    LOGGER.error(f"Failed to set Media Profile: {e}")
+        else:
+            # "(None)" is selected - clear the media profile
+            exec_str = "MediaProfile.Clear"
+            try:
+                devices[0].console_exec_cluster(devices, exec_str)
+                LOGGER.info("Cleared Media Profile")
+            except Exception as e:
+                LOGGER.error(f"Failed to clear Media Profile: {e}")
+
+    def on_set_livelink_preset(self):
+        '''Sets the selected Live Link Preset on the cluster or removes all sources if (None) is selected'''
+        devices = self.get_connected_devices()
+        if not devices:
+            LOGGER.warning("No connected nDisplay devices found")
+            return
+
+        # Get selected livelink preset gamepath
+        current_index = self.livelink_preset_combo.currentIndex()
+        if current_index > 0:  # Non-None selection
+            gamepath = self.livelink_preset_combo.currentData()
+            if gamepath:
+                # Use the helper method from DeviceUnreal to format the command
+                try:
+                    # Create a temporary device instance to use the helper method
+                    temp_device = devices[0]
+                    exec_str = temp_device.exec_command_for_livelink_preset(gamepath)
+
+                    devices[0].console_exec_cluster(devices, exec_str)
+                    LOGGER.info(f"Set Live Link Preset: {gamepath}")
+                except Exception as e:
+                    LOGGER.error(f"Failed to set Live Link Preset: {e}")
+        else:
+            # "(None)" is selected - clear all LiveLink sources
+            exec_str = "LiveLink.Preset.Clear"
+            try:
+                devices[0].console_exec_cluster(devices, exec_str)
+                LOGGER.info("Cleared all Live Link sources")
+            except Exception as e:
+                LOGGER.error(f"Failed to clear Live Link sources: {e}")
\ No newline at end of file
diff --git a/Engine/Plugins/VirtualProduction/Switchboard/Source/Switchboard/switchboard/devices/ndisplay/plugin_ndisplay.py b/Engine/Plugins/VirtualProduction/Switchboard/Source/Switchboard/switchboard/devices/ndisplay/plugin_ndisplay.py
--- a/Engine/Plugins/VirtualProduction/Switchboard/Source/Switchboard/switchboard/devices/ndisplay/plugin_ndisplay.py
+++ b/Engine/Plugins/VirtualProduction/Switchboard/Source/Switchboard/switchboard/devices/ndisplay/plugin_ndisplay.py
@@ -36,6 +36,7 @@ from switchboard.devices.unreal.plugin_unreal import ProgramStartQueueItem, Unre
 
 from .ndisplay_monitor_ui import nDisplayMonitorUI
 from .ndisplay_monitor import nDisplayMonitor
+from .ndisplay_tools_ui import nDisplayToolsUI
 
 
 class PackagingClientConfig(Enum):
@@ -722,6 +723,7 @@ class DevicenDisplay(DeviceUnreal):
 
     ndisplay_monitor_ui = None
     ndisplay_monitor = None
+    ndisplay_tools_ui = None
 
     def __init__(self, name, address, **kwargs):
         super().__init__(name, address, **kwargs)
@@ -1908,8 +1910,14 @@ class DevicenDisplay(DeviceUnreal):
             cls.ndisplay_monitor_ui.replace_exec_history(DevicenDisplay.csettings['cluster_exec_history'].get_value())
             cls.ndisplay_monitor_ui.signal_exec_history_changed.connect(lambda: cls._handle_exec_history_changed())
 
-        # Add our monitor UI to the main tabs in the UI
+        # Create Tools UI if it doesn't exist
+        if not cls.ndisplay_tools_ui:
+            cls.ndisplay_tools_ui = nDisplayToolsUI(
+                parent=tabs, plugin_ndisplay=cls)
+
+        # Add our UIs to the main tabs in the UI
         tabs.addTab(cls.ndisplay_monitor_ui, 'nDisplay &Monitor')
+        tabs.addTab(cls.ndisplay_tools_ui, 'nDisplay &Tools')
 
     @classmethod
     def _handle_exec_history_changed(cls):
diff --git a/Engine/Plugins/VirtualProduction/Switchboard/Source/Switchboard/switchboard/devices/unreal/plugin_unreal.py b/Engine/Plugins/VirtualProduction/Switchboard/Source/Switchboard/switchboard/devices/unreal/plugin_unreal.py
--- a/Engine/Plugins/VirtualProduction/Switchboard/Source/Switchboard/switchboard/devices/unreal/plugin_unreal.py
+++ b/Engine/Plugins/VirtualProduction/Switchboard/Source/Switchboard/switchboard/devices/unreal/plugin_unreal.py
@@ -854,6 +854,8 @@ class DeviceUnreal(Device):
     class StaticSignals(QtCore.QObject):
         ugs_config_updated_signal = QtCore.Signal(ugs_utils.IniParser)
         clipboard_level_added_signal = QtCore.Signal(str)
+        clipboard_assets_added_signal = QtCore.Signal()
+        project_assets_updated_signal = QtCore.Signal()
 
     static_signals = StaticSignals()
 
@@ -3678,6 +3680,9 @@ class DeviceUnreal(Device):
         # save to the cache
         SBCache().update_project_assets(project=project, assets=assets)
 
+        # Emit signal to notify UI components that project assets were updated
+        cls.static_signals.project_assets_updated_signal.emit()
+
     @classmethod
     def handle_clipboard_mime_text(cls, text:str):
         ''' Catches pasting of content browser assets.
@@ -3803,10 +3808,14 @@ class DeviceUnreal(Device):
                 QtWidgets.QMessageBox.Ok | QtWidgets.QMessageBox.Cancel
             )
 
-            SBCache().merge_project_assets(project, assets)
+            if reply == QtWidgets.QMessageBox.Ok:
+                SBCache().merge_project_assets(project, assets)
+
+                for asset in assets:
+                    LOGGER.info(f"Added '{asset.name}'")
 
-            for asset in assets:
-                LOGGER.info(f"Added '{asset.name}'")
+                # Emit signal to notify UI components that assets were added
+                cls.static_signals.clipboard_assets_added_signal.emit()
 
 
 def parse_unreal_tag_file(file_content):
