Switchboard: Made the new MP and LL preset dropdowns non-editable.

Also minor import cleanup.

(Cherrypick of CL 46045846.)

#rnx
#Synced-CL 48062107

diff --git a/Engine/Plugins/VirtualProduction/Switchboard/Source/Switchboard/switchboard/devices/ndisplay/ndisplay_tools_ui.py b/Engine/Plugins/VirtualProduction/Switchboard/Source/Switchboard/switchboard/devices/ndisplay/ndisplay_tools_ui.py
--- a/Engine/Plugins/VirtualProduction/Switchboard/Source/Switchboard/switchboard/devices/ndisplay/ndisplay_tools_ui.py
+++ b/Engine/Plugins/VirtualProduction/Switchboard/Source/Switchboard/switchboard/devices/ndisplay/ndisplay_tools_ui.py
@@ -1,16 +1,14 @@
 # Copyright Epic Games, Inc. All Rights Reserved.
 
-from typing import Optional
-
-from PySide6 import QtCore, QtGui, QtWidgets
+from PySide6 import QtGui, QtWidgets
 from PySide6.QtWidgets import (
-    QComboBox, QHBoxLayout, QLabel, QPushButton,
+    QHBoxLayout, QPushButton,
     QSizePolicy, QVBoxLayout, QWidget, QGroupBox
 )
 
 from switchboard import switchboard_widgets as sb_widgets
 from switchboard.config import CONFIG
-from switchboard.sbcache import SBCache, Asset
+from switchboard.sbcache import SBCache
 from switchboard.switchboard_logging import LOGGER
 from switchboard.devices.unreal.plugin_unreal import DeviceUnreal
 
@@ -123,7 +121,7 @@ class nDisplayToolsUI(QWidget):
 
         # Media Profile combo box
         self.media_profile_combo = sb_widgets.SearchableComboBox(self)
-        self.media_profile_combo.setEditable(True)
+        self.media_profile_combo.setEditable(False)
         self.media_profile_combo.setSizePolicy(
             QSizePolicy.Expanding, QSizePolicy.Preferred)
         self.media_profile_combo.setMinimumWidth(300)
@@ -164,7 +162,7 @@ class nDisplayToolsUI(QWidget):
 
         # Live Link Preset combo box
         self.livelink_preset_combo = sb_widgets.SearchableComboBox(self)
-        self.livelink_preset_combo.setEditable(True)
+        self.livelink_preset_combo.setEditable(False)
         self.livelink_preset_combo.setSizePolicy(
             QSizePolicy.Expanding, QSizePolicy.Preferred)
         self.livelink_preset_combo.setMinimumWidth(300)
@@ -281,7 +279,7 @@ class nDisplayToolsUI(QWidget):
 
         # Get selected media profile gamepath
         current_index = self.media_profile_combo.currentIndex()
-        if current_index > 0:  # Non-None selection
+        if current_index > 0:  # (None) is 0
             gamepath = self.media_profile_combo.currentData()
             if gamepath:
                 # Ensure proper formatting of the game path
@@ -336,4 +334,4 @@ class nDisplayToolsUI(QWidget):
                 devices[0].console_exec_cluster(devices, exec_str)
                 LOGGER.info("Cleared all Live Link sources")
             except Exception as e:
-                LOGGER.error(f"Failed to clear Live Link sources: {e}")
\ No newline at end of file
+                LOGGER.error(f"Failed to clear Live Link sources: {e}")
