Switchboard: Soft termination of LLH.

This will give it a chance to gracefully disconnect from its clients.

(Cherrypick of CL 46049504.)

#Synced-CL 48062107

diff --git a/Engine/Plugins/VirtualProduction/Switchboard/Source/Switchboard/switchboard/devices/livelinkhub/plugin_livelinkhub.py b/Engine/Plugins/VirtualProduction/Switchboard/Source/Switchboard/switchboard/devices/livelinkhub/plugin_livelinkhub.py
--- a/Engine/Plugins/VirtualProduction/Switchboard/Source/Switchboard/switchboard/devices/livelinkhub/plugin_livelinkhub.py
+++ b/Engine/Plugins/VirtualProduction/Switchboard/Source/Switchboard/switchboard/devices/livelinkhub/plugin_livelinkhub.py
@@ -10,7 +10,7 @@ from uuid import UUID
 from PySide6 import QtCore, QtWidgets
 
 from switchboard import message_protocol
-from switchboard.config import CONFIG, Setting, BoolSetting, FilePathSetting, StringSetting
+from switchboard.config import CONFIG, Setting, BoolSetting, FilePathSetting, FloatSetting, StringSetting
 from switchboard.devices.device_base import Device, DeviceStatus, PluginHeaderWidgets
 from switchboard.devices.unreal.plugin_unreal import DeviceUnreal, DeviceWidgetUnreal, ProgramStartQueueItem
 from switchboard.switchboard_logging import LOGGER
@@ -53,6 +53,17 @@ class DeviceLiveLinkHub(DeviceUnreal):
             tool_tip='When checked, retrieves logs after Live Link Hub terminates',
             category='General Settings',
         ),
+        'soft_kill_timeout': FloatSetting(
+            attr_name='soft_kill_timeout',
+            nice_name='Soft Kill Timeout',
+            value=3.0,
+            tool_tip=(
+                'Timeout in seconds for graceful process shutdown before force termination. '
+                'Live Link Hub processes have this amount of time to save state and clean up before being forcefully killed.'
+            ),
+            category='General Settings',
+            show_ui=False,  # For now, avoiding the extra clutter.
+        ),
         'disable_crash_recovery': BoolSetting(
             attr_name='disable_crash_recovery',
             nice_name='Disable Crash Recovery',
@@ -207,15 +218,19 @@ class DeviceLiveLinkHub(DeviceUnreal):
         )
 
     #@override
-    def close(self, force=False):
+    def close(self, force=False, soft_kill_timeout=None):
         llh_puuids = self.program_start_queue.running_puuids_named(self.LLH_PROG_NAME)
 
+        # Get soft kill timeout for this device (use provided value or setting)
+        if soft_kill_timeout is None:
+            soft_kill_timeout = DeviceLiveLinkHub.csettings['soft_kill_timeout'].get_value(self.name)
+
         if not llh_puuids:
             self.status = DeviceStatus.CLOSED
         else:
             self.status = DeviceStatus.CLOSING
             for llh_puuid in llh_puuids:
-                _, msg = message_protocol.create_kill_process_message(llh_puuid)
+                _, msg = message_protocol.create_kill_process_message(llh_puuid, soft_kill_timeout)
                 self.unreal_client.send_message(msg)
 
     #@override
