Switchboard: Make timecode provider reinitializable from Switchboard.

(Cherrypick of CL 46049988.)

#Synced-CL 48062107

diff --git a/Engine/Plugins/VirtualProduction/Switchboard/Source/Switchboard/switchboard/devices/ndisplay/ndisplay_tools_ui.py b/Engine/Plugins/VirtualProduction/Switchboard/Source/Switchboard/switchboard/devices/ndisplay/ndisplay_tools_ui.py
--- a/Engine/Plugins/VirtualProduction/Switchboard/Source/Switchboard/switchboard/devices/ndisplay/ndisplay_tools_ui.py
+++ b/Engine/Plugins/VirtualProduction/Switchboard/Source/Switchboard/switchboard/devices/ndisplay/ndisplay_tools_ui.py
@@ -36,6 +36,10 @@ class nDisplayToolsUI(QWidget):
         livelink_preset_group = self.create_livelink_preset_section()
         layout.addWidget(livelink_preset_group)
 
+        # Timecode section
+        timecode_group = self.create_timecode_section()
+        layout.addWidget(timecode_group)
+
         # Add stretch to push everything to the top
         layout.addStretch()
 
@@ -193,6 +197,30 @@ class nDisplayToolsUI(QWidget):
 
         return group_box
 
+    def create_timecode_section(self) -> QGroupBox:
+        '''Creates the Timecode section of the UI'''
+        group_box = QGroupBox("Timecode")
+        layout = QVBoxLayout()
+
+        # Buttons layout
+        controls_layout = QHBoxLayout()
+
+        # Reinitialize button
+        self.timecode_reinitialize_btn = QPushButton("Reinitialize Provider")
+        self.timecode_reinitialize_btn.setToolTip(
+            "Reinitialize the current timecode provider on the nDisplay cluster")
+        self.timecode_reinitialize_btn.clicked.connect(
+            self.on_reinitialize_timecode_provider)
+        controls_layout.addWidget(self.timecode_reinitialize_btn)
+
+        # Add stretch to push button to the left
+        controls_layout.addStretch()
+
+        layout.addLayout(controls_layout)
+        group_box.setLayout(layout)
+
+        return group_box
+
     def refresh_media_profiles(self):
         '''Refreshes the Media Profile combo box with assets from the cache'''
         self.media_profile_combo.clear()
@@ -335,3 +363,17 @@ class nDisplayToolsUI(QWidget):
                 LOGGER.info("Cleared all Live Link sources")
             except Exception as e:
                 LOGGER.error(f"Failed to clear Live Link sources: {e}")
+
+    def on_reinitialize_timecode_provider(self):
+        '''Reinitializes the timecode provider on the nDisplay cluster'''
+        devices = self.get_connected_devices()
+        if not devices:
+            LOGGER.warning("No connected nDisplay devices found")
+            return
+
+        exec_str = "TimecodeProvider.Reinitialize"
+        try:
+            devices[0].console_exec_cluster(devices, exec_str)
+            LOGGER.info("Reinitializing timecode provider")
+        except Exception as e:
+            LOGGER.error(f"Failed to reinitialize timecode provider: {e}")
diff --git a/Engine/Source/Runtime/Engine/Private/Engine/TimecodeProvider.cpp b/Engine/Source/Runtime/Engine/Private/Engine/TimecodeProvider.cpp
--- a/Engine/Source/Runtime/Engine/Private/Engine/TimecodeProvider.cpp
+++ b/Engine/Source/Runtime/Engine/Private/Engine/TimecodeProvider.cpp
@@ -1,9 +1,25 @@
 // Copyright Epic Games, Inc. All Rights Reserved.
 
 #include "Engine/TimecodeProvider.h"
+#include "Engine/Engine.h"
 
 #include UE_INLINE_GENERATED_CPP_BY_NAME(TimecodeProvider)
 
+namespace UE::Time
+{
+	static FAutoConsoleCommand TimecodeProviderReinitializeCmd(
+		TEXT("TimecodeProvider.Reinitialize"),
+		TEXT("Reinitialize the current timecode provider"),
+		FConsoleCommandDelegate::CreateLambda([]()
+			{
+				if (GEngine)
+				{
+					GEngine->ReinitializeTimecodeProvider();
+				}
+			})
+	);
+}
+
 
 FQualifiedFrameTime UTimecodeProvider::GetDelayedQualifiedFrameTime() const
 {
