Switchboad: Live log viewer support.

SB can request the log file to be streamed. SBL will do so using the normal comms channel, will compress if the data exceeds 1KB, but will send in compressed chunks of 32KB (before compression, expect 4x ratios, up to 10x if repetitive) at a controlled rate.

SB will mirror that remote log file in a local log file. Only one live log is viewable at a time, currently by design with the intent of mitigating any bandwidth issues if SB is receiving too many logs from several remote SBLs.

Bumped SBL version to 3.4.0 (from 3.3.0) as we need it to support live log streaming.

(Modified cherrypick of CLs 46055976, 46056238.)

#Synced-CL 48175258

diff --git a/Engine/Plugins/VirtualProduction/Switchboard/Source/Switchboard/switchboard/devices/unreal/log_viewer_watcher.py b/Engine/Plugins/VirtualProduction/Switchboard/Source/Switchboard/switchboard/devices/unreal/log_viewer_watcher.py
new file mode 100644
--- /dev/null
+++ b/Engine/Plugins/VirtualProduction/Switchboard/Source/Switchboard/switchboard/devices/unreal/log_viewer_watcher.py
@@ -0,0 +1,174 @@
+# Copyright Epic Games, Inc. All Rights Reserved.
+
+import subprocess
+import threading
+import time
+import weakref
+from typing import Dict, Optional
+
+from PySide6 import QtCore
+
+from switchboard.switchboard_logging import LOGGER
+
+
+class LogViewerProcessInfo:
+    '''Holds information about a log viewer process and its associated device.
+    This information is needed to know what to do when the process closes. e.g. End live stream. '''
+
+    def __init__(self, process: subprocess.Popen, device_name: str, device_ref: weakref.ReferenceType, log_file_path: str):
+        self.process = process
+        self.device_name = device_name
+        self.device_ref = device_ref  # Weak reference to device
+        self.log_file_path = log_file_path
+        self.launch_time = time.time()
+
+
+class LogViewerWatcher(QtCore.QObject):
+    ''' Monitors log viewer processes and notifies when they close.
+    Currently expected to be a singleton, declared in this very module.
+    '''
+
+    # Signal emitted when a process terminates. str argument is the device name.
+    process_terminated = QtCore.Signal(str)
+
+    def __init__(self):
+        super().__init__()
+        self._processes: Dict[int, LogViewerProcessInfo] = {}  # key is PID.
+        self._monitor_timer = QtCore.QTimer()
+        self._monitor_timer.timeout.connect(self._check_processes)
+        self._monitor_timer.setInterval(2000)  # Check every 2 seconds
+        self._lock = threading.Lock()
+
+    def start_monitoring(self):
+        ''' Start the process monitoring timer '''
+
+        if not self._monitor_timer.isActive():
+            self._monitor_timer.start()
+
+    def stop_monitoring(self):
+        ''' Stop the process monitoring timer '''
+        if self._monitor_timer.isActive():
+            self._monitor_timer.stop()
+
+    def register_process(self, process: subprocess.Popen, device_name: str, device, log_file_path: str):
+        ''' Register a log viewer process for monitoring '''
+        with self._lock:
+            try:
+                device_ref = weakref.ref(device)
+                process_info = LogViewerProcessInfo(process, device_name, device_ref, log_file_path)
+                self._processes[process.pid] = process_info
+
+                # Start monitoring if this is the first process
+                if len(self._processes) == 1:
+                    self.start_monitoring()
+
+            except Exception as e:
+                LOGGER.error(f"LogViewerWatcher: Failed to register process: {e}")
+
+    def unregister_process(self, pid: int):
+        ''' Manually unregister a process (e.g., when stopping live log) '''
+
+        with self._lock:
+            if pid in self._processes:
+
+                del self._processes[pid]
+
+                # Stop monitoring if no processes left
+                if not self._processes:
+                    self.stop_monitoring()
+
+    def terminate_process(self, pid: int) -> bool:
+        ''' Terminate a specific process by PID '''
+
+        with self._lock:
+
+            if pid in self._processes:
+
+                process_info = self._processes[pid]
+
+                try:
+                    if process_info.process.poll() is None:  # Process still running
+                        process_info.process.terminate()
+                        return True
+
+                except Exception as e:
+                    LOGGER.error(f"LogViewerWatcher: Failed to terminate process PID {pid}: {e}")
+
+            return False
+
+    def terminate_log_viewers_for_device(self, device_name: str):
+        ''' Terminate all processes associated with a specific device '''
+
+        pids_to_terminate = []
+
+        with self._lock:
+
+            for pid, process_info in self._processes.items():
+                if process_info.device_name == device_name:
+                    pids_to_terminate.append(pid)
+
+        for pid in pids_to_terminate:
+            self.terminate_process(pid)
+
+    def get_process_for_device(self, device_name: str) -> Optional[subprocess.Popen]:
+        ''' Get the process object for a device, if any '''
+
+        with self._lock:
+            for process_info in self._processes.values():
+                if process_info.device_name == device_name:
+                    return process_info.process
+        return None
+
+    @QtCore.Slot()
+    def _check_processes(self):
+        ''' Check all registered processes and clean up terminated ones '''
+
+        terminated_pids = []
+
+        with self._lock:
+            for pid, process_info in self._processes.items():
+                try:
+                    # Check if process has terminated
+                    if process_info.process.poll() is not None:
+                        terminated_pids.append(pid)
+                        
+                except Exception as e:
+                    # Process might be invalid, mark for removal
+                    LOGGER.warning(f"LogViewerWatcher: Error checking process PID {pid}: {e}")
+                    terminated_pids.append(pid)
+
+        # Clean up terminated processes and notify
+
+        for pid in terminated_pids:
+            with self._lock:
+                if pid in self._processes:
+                    process_info = self._processes[pid]
+                    device_name = process_info.device_name
+                    device_ref = process_info.device_ref
+
+                    # Remove from tracking
+                    del self._processes[pid]
+
+                    # Try to get the device object to clean up live log streaming
+                    device = device_ref() if device_ref else None
+
+                    if device:
+                        # Force main thread execution. We're not in a hurry.
+                        QtCore.QMetaObject.invokeMethod(
+                            device, '_on_log_viewer_closed',
+                            QtCore.Qt.QueuedConnection
+                        )
+                    else:
+                        LOGGER.info(f"LogViewerWatcher: Device '{device_name}' no longer exists")
+
+                    # Emit signal for external listeners
+                    self.process_terminated.emit(device_name)
+
+        # Stop monitoring if no processes left
+        with self._lock:
+            if not self._processes:
+                self.stop_monitoring()
+
+
+# Global log viewer watcher instance
+log_viewer_watcher = LogViewerWatcher()
\ No newline at end of file
diff --git a/Engine/Plugins/VirtualProduction/Switchboard/Source/Switchboard/switchboard/devices/unreal/plugin_unreal.py b/Engine/Plugins/VirtualProduction/Switchboard/Source/Switchboard/switchboard/devices/unreal/plugin_unreal.py
--- a/Engine/Plugins/VirtualProduction/Switchboard/Source/Switchboard/switchboard/devices/unreal/plugin_unreal.py
+++ b/Engine/Plugins/VirtualProduction/Switchboard/Source/Switchboard/switchboard/devices/unreal/plugin_unreal.py
@@ -10,11 +10,13 @@ import os
 import pathlib
 import re
 import socket
+import subprocess
 import sys
 import threading
 import time
 import uuid
 import weakref
+import zlib
 
 from abc import abstractmethod
 from collections import OrderedDict
@@ -46,6 +48,7 @@ from switchboard.switchboard_logging import LOGGER
 from switchboard.tools.insights_launcher import InsightsLauncher
 from . import version_helpers
 from .listener_watcher import ListenerWatcher
+from .log_viewer_watcher import log_viewer_watcher
 from .redeploy_dialog import RedeployListenerDialog
 from switchboard.sbcache import SBCache, Asset
 
@@ -517,6 +520,10 @@ class SyncCategoryOption:
             return False
 
 
+# =============
+# Unreal Device
+# =============
+
 class DeviceUnreal(Device):
 
     NDISPLAY_CLASS_NAMES = (
@@ -873,6 +880,9 @@ class DeviceUnreal(Device):
     # used for listener updates.
     active_unreal_devices: set[DeviceUnreal] = set()
 
+    # Global live log management - only one device can stream live logs at a time
+    active_live_log_device: Optional['DeviceUnreal'] = None
+
     # Flag used to batch together multiple rapid calls to
     # `_queue_notify_redeploy`.
     _pending_notify_redeploy = False
@@ -1175,6 +1185,12 @@ class DeviceUnreal(Device):
         self.unreal_client.delegates["program killed"] = (self.on_program_killed)
         # This catches start failures.
         self.unreal_client.delegates["start"] = (self.on_program_started)
+        self.unreal_client.delegates["live log data"] = (self.on_live_log_data)
+
+        # Live log streaming state
+        self.live_log_streaming = False
+        self.live_log_file_path = None
+        self.live_log_local_file = None
 
         self.osc_connection_timer = QtCore.QTimer(self)
         self.osc_connection_timer.timeout.connect(self._try_connect_osc)
@@ -1503,6 +1519,9 @@ class DeviceUnreal(Device):
 
         super().device_widget_registered(device_widget)
 
+        # Store device reference in widget for live log context menu
+        device_widget._device = self
+
         device_widget.signal_exclude_from_build_toggled.connect(self.on_toggle_exclude_from_build)
 
         device_widget.signal_exclude_from_insights_toggled.connect(self.on_toggle_exclude_from_insights)
@@ -1516,6 +1535,9 @@ class DeviceUnreal(Device):
         # hook to copy last launch signal from widget
         device_widget.signal_copy_last_launch_command.connect(self.on_copy_last_launch_command)
 
+        # hook to live log signal from widget
+        device_widget.signal_start_live_log.connect(self.on_start_live_log)
+
         self.update_settings_menu_state()
 
         device_widget.autojoin_mu.signal_device_widget_autojoin_mu.connect(self.on_autojoin_mu_ui_change)
@@ -3460,6 +3482,186 @@ class DeviceUnreal(Device):
         insights = InsightsLauncher()
         insights.launch(args=[f'"{str(tracepath)}"'], allow_duplicate=True)
 
+    def on_start_live_log(self):
+        ''' Starts live log streaming from the remote device. 
+        
+        This ends up being called from the device context menu 
+        when the user wants to view the live log
+        '''
+
+        # Check if another device is already streaming the live log
+
+        if DeviceUnreal.active_live_log_device and DeviceUnreal.active_live_log_device != self:
+            LOGGER.info(f"Stopping live log streaming on {DeviceUnreal.active_live_log_device.name} to start on {self.name}")
+            DeviceUnreal.active_live_log_device._stop_live_log_internal()
+
+        # If already streaming, see if we need to launch an instance of the log viewer.
+
+        if self.live_log_streaming and self.live_log_local_file:
+            if self.live_log_local_file.exists():
+
+                # Check if we have a running process for this device
+                existing_process = log_viewer_watcher.get_process_for_device(self.name)
+
+                if existing_process and existing_process.poll() is None:
+                    LOGGER.info(f"Log viewer for '{self.name}' is already running.")
+                else:
+                    LOGGER.info(f"Launching live log viewer for '{self.name}'.")
+                    self._launch_log_viewer()
+                return
+            else:
+                # Local file was unexpectedly deleted, so we stop it because we need to get the stream from the beginning.
+                self._stop_live_log_internal()
+
+        if not self.unreal_client.is_connected:
+            LOGGER.error(f"Cannot start live log because device {self.name} is not connected")
+            return
+
+        if self.status < DeviceStatus.OPEN:
+            LOGGER.error(f"Cannot start live log because Unreal is not running on device {self.name}")
+            return
+
+        # Construct the remote log file path
+        remote_log_path = self.get_remote_log_path() / self.log_filename
+        LOGGER.info(f"Live log: requesting streaming for remote path: {remote_log_path}")
+
+        # Create local live log file
+        log_dir = DeviceUnreal.get_log_download_dir()
+        timestamp = datetime.now().strftime("%Y%m%d_%H%M%S")
+        local_filename = f"live_{self.name}_{timestamp}.log"
+        self.live_log_local_file = pathlib.Path(log_dir) / local_filename
+
+        # Create the local live log file, that the log viewer will display as it updates.
+        try:
+            self.live_log_local_file.parent.mkdir(parents=True, exist_ok=True)
+            self.live_log_local_file.touch()
+        except Exception as exception:
+            LOGGER.error(f"Failed to create live log file {self.live_log_local_file}: {exception}")
+            return
+
+        # Get the UUID of the running Unreal process (if any)
+
+        unreal_process_uuid = None
+        unreal_puuids = self.program_start_queue.running_puuids_named(UnrealJobs.Unreal.value)
+
+        if unreal_puuids:
+            # Grab the first running Unreal, although we expect there to only be up to one per device.
+            unreal_process_uuid = str(unreal_puuids[0])
+        else:
+            LOGGER.error("Live log: no running Unreal processes found, can't start the live log")
+            return
+
+        # Send start live log message to listener
+
+        msgid, msg = message_protocol.create_start_live_log_stream_message(
+            str(remote_log_path), tail_only=False, process_uuid=unreal_process_uuid)
+
+        try:
+            self.unreal_client.send_message(msg)
+            self.live_log_streaming = True
+            self.live_log_file_path = str(remote_log_path)
+            DeviceUnreal.active_live_log_device = self
+
+            LOGGER.info(f"Started live log streaming for {self.name} -> {self.live_log_local_file}")
+
+            # Launch the log viewer (and will also register with its process watcher)
+            self._launch_log_viewer()
+
+        except Exception as exception:
+            LOGGER.error(f"Failed to start live log streaming: {exception}")
+            self.live_log_streaming = False
+            self.live_log_file_path = None
+
+    def on_live_log_data(self, message):
+        ''' Handles incoming live log data from the listener '''
+
+        if not self.live_log_streaming or not self.live_log_local_file:
+            LOGGER.warning(f"Received live log data but streaming not active or no local file: streaming={self.live_log_streaming}, file={self.live_log_local_file}")
+            return
+
+        try:
+            # Check for compressed content first
+
+            if 'content_compressed_base64' in message:
+                compressed_base64 = message['content_compressed_base64']
+                compressed_bytes = base64.b64decode(compressed_base64)
+
+                # Decompress
+                content_bytes = zlib.decompress(compressed_bytes)
+
+                # Uncomment the following to get an idea of the compression ratios achieved.
+                #
+                #compression_ratio = len(content_bytes) / len(compressed_bytes) if len(compressed_bytes) > 0 else 1.0
+                #LOGGER.info(f"Live log compression: {len(compressed_bytes)} -> {len(content_bytes)} bytes (ratio: {compression_ratio:.2f}x)")
+
+                # write to live log file
+                with open(self.live_log_local_file, 'ab') as f:
+                    f.write(content_bytes)
+
+            elif 'content_base64' in message:
+
+                # If data is not compressed, check the uncompressed content
+
+                content_base64 = message['content_base64']
+                content_bytes = base64.b64decode(content_base64)
+
+                # write to live log file
+                with open(self.live_log_local_file, 'ab') as f:
+                    f.write(content_bytes)
+
+        except Exception as e:
+            LOGGER.error(f"Failed to write live log data to {self.live_log_local_file}: {e}")
+
+    def _launch_log_viewer(self):
+        ''' Launch the log viewer with the live log file '''
+
+        # Currently default viewer for this filetype.
+        url = QtCore.QUrl.fromLocalFile(str(self.live_log_local_file))
+        QtGui.QDesktopServices.openUrl(url)
+
+    def _stop_live_log_internal(self):
+        ''' Stops live log streaming '''
+
+        if not self.live_log_streaming:
+            return
+
+        if self.live_log_file_path and self.unreal_client.is_connected:
+
+            # Send stop live log message to listener
+            msgid, msg = message_protocol.create_stop_live_log_stream_message(
+                self.live_log_file_path)
+
+            try:
+                self.unreal_client.send_message(msg)
+                LOGGER.info(f"Stopped live log streaming for {self.name}")
+            except Exception as e:
+                LOGGER.error(f"Failed to send stop live log message: {e}")
+
+        # Terminate any running log viewer processes for this device
+        log_viewer_watcher.terminate_log_viewers_for_device(self.name)
+
+        # Reset state
+        self.live_log_streaming = False
+        self.live_log_file_path = None
+        self.live_log_local_file = None
+
+        # Clear global active device, if it was this device.
+        if DeviceUnreal.active_live_log_device == self:
+            DeviceUnreal.active_live_log_device = None
+
+    @QtCore.Slot()
+    def _on_log_viewer_closed(self):
+        '''
+        Called by the log viewer watcer when the log viewer process associated with this device terminates,
+        e.g. the user closed it. When this happens, we automatically stop the live streaming since nobody is
+        watching. Currently there is not support for log viewers different to the default.
+        '''
+        LOGGER.info(f"Log viewer closed for device '{self.name}', stopping live log streaming")
+
+        # Stop live log streaming since the viewer is no longer available
+        if self.live_log_streaming:
+            self._stop_live_log_internal()
+
     def on_copy_last_launch_command(self):
         ''' Copies the last launch command to the clipboard'''
         QtGui.QGuiApplication.clipboard().setText(
@@ -3841,6 +4043,7 @@ class DeviceWidgetUnreal(DeviceWidget):
     signal_open_last_trace = QtCore.Signal(object)
     signal_copy_last_launch_command = QtCore.Signal(object)
     signal_device_widget_fill_ddc = QtCore.Signal(object, bool)
+    signal_start_live_log = QtCore.Signal(object)
 
     def __init__(self, name, device_hash, address, icons, parent=None):
         self._autojoin_visible = True
@@ -4269,21 +4472,36 @@ class DeviceWidgetUnreal(DeviceWidget):
     def populate_context_menu(self, cmenu: QtWidgets.QMenu):
         ''' Called to populate the given context menu with any desired actions '''
 
-        cmenu.addAction(
-            "Include in build" if self.exclude_from_build else "Exclude from build",
-            lambda: self.signal_exclude_from_build_toggled.emit()
-        )
-        cmenu.addAction(
-            "Include in Insights traces" if self.exclude_from_insights else "Exclude from Insights traces",
-            lambda: self.signal_exclude_from_insights_toggled.emit()
-        )
+        # Remove device is already added by parent widget
+        cmenu.addSeparator()
+
+        # Exclude options with checkboxes
+        exclude_build_action = cmenu.addAction("Exclude from build")
+        exclude_build_action.setCheckable(True)
+        exclude_build_action.setChecked(self.exclude_from_build)
+        exclude_build_action.triggered.connect(lambda: self.signal_exclude_from_build_toggled.emit())
+
+        exclude_insights_action = cmenu.addAction("Exclude from Insights traces")
+        exclude_insights_action.setCheckable(True)
+        exclude_insights_action.setChecked(self.exclude_from_insights)
+        exclude_insights_action.triggered.connect(lambda: self.signal_exclude_from_insights_toggled.emit())
+
+        cmenu.addSeparator()
+
+        # Log and trace options
         cmenu.addAction("Open fetched log", lambda: self.signal_open_last_log.emit(self))
+
+        # View Live Log - only if Unreal is running
+        if hasattr(self, '_device') and self._device and self._device.status >= DeviceStatus.OPEN:
+            cmenu.addAction("View Live Log", lambda: self.signal_start_live_log.emit(self))
+
         cmenu.addAction("Open fetched trace", lambda: self.signal_open_last_trace.emit(self))
         cmenu.addAction("Copy last launch command", lambda: self.signal_copy_last_launch_command.emit(self))
 
-        # Only create DDC submenu if node is connected
+        cmenu.addSeparator()
+
+        # Fill DDC option - only if node is connected
         if self.connect_button.isChecked():
             fill_ddc_menu = cmenu.addMenu("Fill DDC (Prepare Shaders)")
-
             fill_ddc_menu.addAction("Current Level", lambda: self.signal_device_widget_fill_ddc.emit(self, True))
             fill_ddc_menu.addAction("All Levels", lambda: self.signal_device_widget_fill_ddc.emit(self, False))
diff --git a/Engine/Plugins/VirtualProduction/Switchboard/Source/Switchboard/switchboard/devices/unreal/version_helpers.py b/Engine/Plugins/VirtualProduction/Switchboard/Source/Switchboard/switchboard/devices/unreal/version_helpers.py
--- a/Engine/Plugins/VirtualProduction/Switchboard/Source/Switchboard/switchboard/devices/unreal/version_helpers.py
+++ b/Engine/Plugins/VirtualProduction/Switchboard/Source/Switchboard/switchboard/devices/unreal/version_helpers.py
@@ -3,7 +3,7 @@
 import typing
 
 
-LISTENER_COMPATIBLE_VERSION = (3,3)
+LISTENER_COMPATIBLE_VERSION = (3,4)
 
 # Version where redeploy support was initially introduced.
 LISTENER_MIN_REDEPLOY_VERSION = (1,4,0)
diff --git a/Engine/Plugins/VirtualProduction/Switchboard/Source/Switchboard/switchboard/listener_client.py b/Engine/Plugins/VirtualProduction/Switchboard/Source/Switchboard/switchboard/listener_client.py
--- a/Engine/Plugins/VirtualProduction/Switchboard/Source/Switchboard/switchboard/listener_client.py
+++ b/Engine/Plugins/VirtualProduction/Switchboard/Source/Switchboard/switchboard/listener_client.py
@@ -344,6 +344,9 @@ class ListenerClient(QtCore.QObject):
         self.delegates: Dict[str, Optional[Callable[[Dict], None]]] = {
             "state": None,
             "get sync status": None,
+            "live log data": None,
+            "start live log stream": None,
+            "stop live log stream": None,
         }
 
         self.async_msg_states: dict[uuid.UUID, dict] = {}
@@ -583,6 +586,28 @@ class ListenerClient(QtCore.QObject):
                 if self.receive_file_failed_delegate:
                     self.receive_file_failed_delegate(message['source'],
                                                       message['error'])
+
+        elif "start live log stream" in message:
+            if message['start live log stream'] is True:
+                LOGGER.debug("Live log stream started successfully")
+            else:
+                LOGGER.warning(f"Failed to start live log stream: {message.get('error', 'Unknown error')}")
+
+        elif "stop live log stream" in message:
+            if message['stop live log stream'] is True:
+                LOGGER.debug("Live log stream stopped successfully")
+            else:
+                # This may happen when SBL auto-stops the stream when the UE process dies before Python tries to stop it.
+                LOGGER.debug(f"Stop live log stream response: {message.get('error', 'Unknown error')}")
+
+        elif "live log stream terminated" in message:
+            if message['live log stream terminated'] is True:
+                log_file_path = message.get('log_file_path', 'unknown')
+                reason = message.get('reason', 'unknown reason')
+                LOGGER.info(f"Live log stream terminated for {log_file_path}: {reason}")
+            else:
+                LOGGER.warning(f"Live log stream termination notification failed: {message.get('error', 'Unknown error')}")
+
         else:
             LOGGER.error(f'Unhandled message: {message}')
             raise ValueError
diff --git a/Engine/Plugins/VirtualProduction/Switchboard/Source/Switchboard/switchboard/message_protocol.py b/Engine/Plugins/VirtualProduction/Switchboard/Source/Switchboard/switchboard/message_protocol.py
--- a/Engine/Plugins/VirtualProduction/Switchboard/Source/Switchboard/switchboard/message_protocol.py
+++ b/Engine/Plugins/VirtualProduction/Switchboard/Source/Switchboard/switchboard/message_protocol.py
@@ -231,3 +231,32 @@ def create_set_inactive_timeout_message(timeout_seconds: int):
            'seconds': timeout_seconds}
     message = json.dumps(cmd).encode() + b'\x00'
     return (cmd_id, message)
+
+
+def create_start_live_log_stream_message(log_file_path: str, tail_only: bool = False, process_uuid: str = None):
+    ''' Creates the message to start SBL streaming the live log of the running Unreal/nDisplay process '''
+
+    cmd_id = uuid.uuid4()
+    cmd = {
+        'command': 'start live log stream',
+        'id': str(cmd_id),
+        'log_file_path': log_file_path,
+        'tail_only': tail_only
+    }
+    if process_uuid:
+        cmd['process_uuid'] = process_uuid
+    message = json.dumps(cmd).encode() + b'\x00'
+    return (cmd_id, message)
+
+
+def create_stop_live_log_stream_message(log_file_path: str):
+    ''' Creates the message to stop SBL streaming the live log of the running Unreal/nDisplay process. '''
+
+    cmd_id = uuid.uuid4()
+    cmd = {
+        'command': 'stop live log stream',
+        'id': str(cmd_id),
+        'log_file_path': log_file_path
+    }
+    message = json.dumps(cmd).encode() + b'\x00'
+    return (cmd_id, message)
diff --git a/Engine/Source/Programs/SwitchboardListener/SblCore/Internal/SwitchboardListener.h b/Engine/Source/Programs/SwitchboardListener/SblCore/Internal/SwitchboardListener.h
--- a/Engine/Source/Programs/SwitchboardListener/SblCore/Internal/SwitchboardListener.h
+++ b/Engine/Source/Programs/SwitchboardListener/SblCore/Internal/SwitchboardListener.h
@@ -39,6 +39,8 @@ struct FSwitchboardFixExeFlagsTask;
 struct FSwitchboardMinimizeWindowsTask;
 struct FSwitchboardSetInactiveTimeoutTask;
 struct FSwitchboardFreeListenerBinaryTask;
+struct FSwitchboardStartLiveLogStreamTask;
+struct FSwitchboardStopLiveLogStreamTask;
 
 class FGpuClockerManager;
 
@@ -160,6 +162,12 @@ private:
 	bool Task_SetInactiveTimeout(const FSwitchboardSetInactiveTimeoutTask& InTimeoutTask);
 	bool Task_FreeListenerBinary(const FSwitchboardFreeListenerBinaryTask& InRenameProcessTask);
 
+	/** Starts streaming a log file to Switchboard */
+	bool Task_StartLiveLogStream(const FSwitchboardStartLiveLogStreamTask& InStartLiveLogStreamTask);
+
+	/** Stops streaming a log file to Switchboard */
+	bool Task_StopLiveLogStream(const FSwitchboardStopLiveLogStreamTask& InStopLiveLogStreamTask);
+
 	bool KillProcessNow(FRunningProcess* InProcess, float SoftKillTimeout = 0.0f);
 
 #if PLATFORM_WINDOWS
@@ -178,6 +186,9 @@ private:
 	void HandleStdout(const TSharedPtr<FRunningProcess>& Process);
 	void HandleRunningProcesses(TArray<TSharedPtr<FRunningProcess>>& Processes, bool bNotifyThatProgramEnded);
 
+	/** Processes all active live log streams, reading new content and sending to Switchboard clients */
+	void HandleLiveLogStreams();
+
 	bool SendMessage(const FString& InMessage, const FIPv4Endpoint& InEndpoint);
 	void SendMessageFutures();
 
@@ -294,4 +305,38 @@ private:
 	FSimpleMulticastDelegate OnInitDelegate;
 	FSimpleMulticastDelegate OnShutdownDelegate;
 	FSimpleMulticastDelegate OnTickDelegate;
+
+	/** Live log streaming state */
+	struct FLiveLogStream
+	{
+		/** Local path to the log file being streamed */
+		FString LogFilePath;
+
+		/** Network endpoint of the client receiving the stream */
+		FIPv4Endpoint ClientEndpoint;
+
+		/** Current read position in the log file */
+		int64 LastFilePosition;
+
+		/** If true, only stream new content added after the stream is started */
+		bool bTailOnly;
+
+		/** Timestamp of the last stream update check */
+		double LastUpdateTime;
+
+		/** UUID of the process whose log is being streamed */
+		FGuid AssociatedProcessUUID;
+
+		FLiveLogStream(const FString& InLogFilePath, const FIPv4Endpoint& InClientEndpoint, bool bInTailOnly, const FGuid& InProcessUUID = FGuid())
+			: LogFilePath(InLogFilePath)
+			, ClientEndpoint(InClientEndpoint)
+			, LastFilePosition(0)
+			, bTailOnly(bInTailOnly)
+			, LastUpdateTime(0.0)
+			, AssociatedProcessUUID(InProcessUUID)
+		{}
+	};
+
+	TArray<FLiveLogStream> ActiveLogStreams;
+	UE::FMutex LiveLogStreamsMutex;
 };
diff --git a/Engine/Source/Programs/SwitchboardListener/SblCore/Internal/SwitchboardListenerVersion.h b/Engine/Source/Programs/SwitchboardListener/SblCore/Internal/SwitchboardListenerVersion.h
--- a/Engine/Source/Programs/SwitchboardListener/SblCore/Internal/SwitchboardListenerVersion.h
+++ b/Engine/Source/Programs/SwitchboardListener/SblCore/Internal/SwitchboardListenerVersion.h
@@ -3,5 +3,5 @@
 #pragma once
 
 #define SBLISTENER_VERSION_MAJOR 3
-#define SBLISTENER_VERSION_MINOR 3
+#define SBLISTENER_VERSION_MINOR 4
 #define SBLISTENER_VERSION_PATCH 0
diff --git a/Engine/Source/Programs/SwitchboardListener/SblCore/Private/SwitchboardListener.cpp b/Engine/Source/Programs/SwitchboardListener/SblCore/Private/SwitchboardListener.cpp
--- a/Engine/Source/Programs/SwitchboardListener/SblCore/Private/SwitchboardListener.cpp
+++ b/Engine/Source/Programs/SwitchboardListener/SblCore/Private/SwitchboardListener.cpp
@@ -17,6 +17,7 @@
 #include "Async/Async.h"
 #include "Async/AsyncWork.h"
 #include "Common/TcpListener.h"
+#include "Dom/JsonObject.h"
 #include "GenericPlatform/GenericPlatformMisc.h"
 #include "GenericPlatform/GenericPlatformProcess.h"
 #include "HAL/FileManager.h"
@@ -27,9 +28,12 @@
 #include "Logging/StructuredLog.h"
 #include "Misc/Base64.h"
 #include "Misc/CommandLine.h"
+#include "Misc/Compression.h"
 #include "Misc/FileHelper.h"
 #include "Misc/Paths.h"
 #include "Misc/ScopeExit.h"
+#include "Serialization/JsonSerializer.h"
+#include "Serialization/JsonWriter.h"
 
 #include <atomic>
 
@@ -642,6 +646,7 @@ void FSwitchboardListener::Tick()
 	CleanUpDisconnectedSockets();
 	HandleRunningProcesses(RunningProcesses, true);
 	HandleRunningProcesses(FlipModeMonitors, false);
+	HandleLiveLogStreams();
 	SendMessageFutures();
 	GpuClockerManager->Tick();
 
@@ -788,6 +793,16 @@ bool FSwitchboardListener::RunScheduledTask(const FSwitchboardTask& InTask)
 			const FSwitchboardFreeListenerBinaryTask& Task = static_cast<const FSwitchboardFreeListenerBinaryTask&>(InTask);
 			return Task_FreeListenerBinary(Task);
 		}
+		case ESwitchboardTaskType::StartLiveLogStream:
+		{
+			const FSwitchboardStartLiveLogStreamTask& Task = static_cast<const FSwitchboardStartLiveLogStreamTask&>(InTask);
+			return Task_StartLiveLogStream(Task);
+		}
+		case ESwitchboardTaskType::StopLiveLogStream:
+		{
+			const FSwitchboardStopLiveLogStreamTask& Task = static_cast<const FSwitchboardStopLiveLogStreamTask&>(InTask);
+			return Task_StopLiveLogStream(Task);
+		}
 		default:
 		{
 			static const FString Response = TEXT("Unknown Command detected");
@@ -2652,6 +2667,276 @@ bool FSwitchboardListener::Task_FreeListenerBinary(const FSwitchboardFreeListene
 	return true;
 }
 
+bool FSwitchboardListener::Task_StartLiveLogStream(const FSwitchboardStartLiveLogStreamTask& InStartLiveLogStreamTask)
+{
+	SWITCHBOARD_TRACE_CPUPROFILER_EVENT_SCOPE(FSwitchboardListener::Task_StartLiveLogStream);
+
+	const FString& LogFilePath = InStartLiveLogStreamTask.LogFilePath;
+	const FIPv4Endpoint& ClientEndpoint = InStartLiveLogStreamTask.Recipient;
+
+	// Check if the given log file exists
+	if (!FPaths::FileExists(LogFilePath))
+	{
+		const FString ErrorMsg = FString::Printf(TEXT("Log file does not exist: %s"), *LogFilePath);
+
+		UE_LOG(LogSwitchboard, Error, TEXT("Start live log stream: %s"), *ErrorMsg);
+
+		SendMessage(CreateTaskDeclinedMessage(InStartLiveLogStreamTask, ErrorMsg, {}), ClientEndpoint);
+
+		return false;
+	}
+
+	// Check if stream already exists for this file and client
+	{
+		UE::TUniqueLock<UE::FMutex> Lock(LiveLogStreamsMutex);
+		for (const FLiveLogStream& Stream : ActiveLogStreams)
+		{
+			if (Stream.LogFilePath == LogFilePath && Stream.ClientEndpoint == ClientEndpoint)
+			{
+				const FString ErrorMsg = TEXT("Live log stream already active for this file and client");
+
+				SendMessage(CreateTaskDeclinedMessage(InStartLiveLogStreamTask, ErrorMsg, {}), ClientEndpoint);
+				return false;
+			}
+		}
+
+		// Add new stream
+		FLiveLogStream NewStream(LogFilePath, ClientEndpoint, InStartLiveLogStreamTask.bTailOnly, InStartLiveLogStreamTask.ProcessUUID);
+
+		if (InStartLiveLogStreamTask.ProcessUUID.IsValid())
+		{
+			UE_LOG(LogSwitchboard, Log, TEXT("Starting live log stream for %s with associated process UUID: %s"), *LogFilePath, *InStartLiveLogStreamTask.ProcessUUID.ToString());
+		}
+		else
+		{
+			UE_LOG(LogSwitchboard, Log, TEXT("Starting live log stream for %s with no associated process UUID"), *LogFilePath);
+		}
+
+		// Start from the beginning unless bTailOnly is true.
+		if (!NewStream.bTailOnly)
+		{
+			NewStream.LastFilePosition = 0;
+		}
+		else
+		{
+			// Start from end of file for tail-only mode
+			TUniquePtr<IFileHandle> FileHandle(FPlatformFileManager::Get().GetPlatformFile().OpenRead(*LogFilePath));
+			if (FileHandle)
+			{
+				NewStream.LastFilePosition = FileHandle->Size();
+			}
+		}
+
+		ActiveLogStreams.Add(MoveTemp(NewStream));
+	}
+
+	// Send success response
+	SendMessage(CreateCommandAcceptedMessage(InStartLiveLogStreamTask.TaskID), ClientEndpoint);
+	UE_LOG(LogSwitchboard, Log, TEXT("Started live log stream for %s"), *LogFilePath);
+
+	return true;
+}
+
+bool FSwitchboardListener::Task_StopLiveLogStream(const FSwitchboardStopLiveLogStreamTask& InStopLiveLogStreamTask)
+{
+	SWITCHBOARD_TRACE_CPUPROFILER_EVENT_SCOPE(FSwitchboardListener::Task_StopLiveLogStream);
+
+	const FString& LogFilePath = InStopLiveLogStreamTask.LogFilePath;
+	const FIPv4Endpoint& ClientEndpoint = InStopLiveLogStreamTask.Recipient;
+
+	bool bStreamFound = false;
+	{
+		UE::TUniqueLock<UE::FMutex> Lock(LiveLogStreamsMutex);
+		for (int32 StreamIdx = ActiveLogStreams.Num() - 1; StreamIdx >= 0; --StreamIdx)
+		{
+			const FLiveLogStream& Stream = ActiveLogStreams[StreamIdx];
+			if (Stream.LogFilePath == LogFilePath && Stream.ClientEndpoint == ClientEndpoint)
+			{
+				ActiveLogStreams.RemoveAt(StreamIdx);
+				bStreamFound = true;
+				break;
+			}
+		}
+	}
+
+	if (bStreamFound)
+	{
+		SendMessage(CreateCommandAcceptedMessage(InStopLiveLogStreamTask.TaskID), ClientEndpoint);
+		UE_LOG(LogSwitchboard, Log, TEXT("Stopped live log stream for %s"), *LogFilePath);
+
+		return true;
+	}
+	else
+	{
+		const FString ErrorMsg = TEXT("No active live log stream found for this file and client");
+		SendMessage(CreateTaskDeclinedMessage(InStopLiveLogStreamTask, ErrorMsg, {}), ClientEndpoint);
+
+		return false;
+	}
+}
+
+void FSwitchboardListener::HandleLiveLogStreams()
+{
+	SWITCHBOARD_TRACE_CPUPROFILER_EVENT_SCOPE(FSwitchboardListener::HandleLiveLogStreams);
+
+	const double CurrentTime = FPlatformTime::Seconds();
+	constexpr double StreamUpdateInterval = 0.1; // 100 ms between updates
+
+	UE::TUniqueLock<UE::FMutex> Lock(LiveLogStreamsMutex);
+
+	// Collect indices of streams whose associated processes have ended. 
+	// If they have, then the stream are removed, effectively stopping them.
+
+	TArray<int32> StreamsToRemove;
+
+	for (int32 StreamIdx = 0; StreamIdx < ActiveLogStreams.Num(); ++StreamIdx)
+	{
+		FLiveLogStream& Stream = ActiveLogStreams[StreamIdx];
+
+		// Check if the associated process is still running (if we have a valid UUID)
+
+		if (Stream.AssociatedProcessUUID.IsValid())
+		{
+			bool bProcessStillRunning = false;
+
+			for (const TSharedPtr<FRunningProcess>& Process : RunningProcesses)
+			{
+				if (Process && Process->UUID == Stream.AssociatedProcessUUID)
+				{
+					bProcessStillRunning = true;
+					break;
+				}
+			}
+
+			if (!bProcessStillRunning)
+			{
+				UE_LOG(LogSwitchboard, Log, TEXT("Stopping live log stream for %s because the associated Unreal process has ended"), *Stream.LogFilePath);
+				StreamsToRemove.Add(StreamIdx);
+				continue;
+			}
+		}
+	}
+
+	// Remove streams whose processes have ended (in reverse order to preserve indices)
+
+	for (int32 StreamIdx = StreamsToRemove.Num() - 1; StreamIdx >= 0; --StreamIdx)
+	{
+		const FLiveLogStream& StreamToRemove = ActiveLogStreams[StreamsToRemove[StreamIdx]];
+
+		// Notify Switchboard that the stream has terminated
+		const FString TerminationMessage = CreateLiveLogStreamTerminatedMessage(StreamToRemove.LogFilePath, TEXT("process terminated"));
+		SendMessage(TerminationMessage, StreamToRemove.ClientEndpoint);
+
+		ActiveLogStreams.RemoveAt(StreamsToRemove[StreamIdx]);
+	}
+
+	// Process remaining active streams
+	for (FLiveLogStream& Stream : ActiveLogStreams)
+	{
+		// Rate limiting: only update every StreamUpdateInterval seconds per stream
+		if (CurrentTime - Stream.LastUpdateTime < StreamUpdateInterval)
+		{
+			continue;
+		}
+
+		Stream.LastUpdateTime = CurrentTime;
+
+		IPlatformFile& PlatformFile = FPlatformFileManager::Get().GetPlatformFile();
+
+		if (!PlatformFile.FileExists(*Stream.LogFilePath))
+		{
+			UE_LOG(LogSwitchboard, Warning, TEXT("Live log file does not exist: %s"), *Stream.LogFilePath);
+			continue;
+		}
+
+		// Use shared read mode to allow reading while Unreal is writing to the file (otherwise open will fail)
+
+		TUniquePtr<IFileHandle> FileHandle(PlatformFile.OpenRead(*Stream.LogFilePath, true));
+
+		if (!FileHandle)
+		{
+			UE_LOG(LogSwitchboard, Warning, TEXT("Live log file cannot be opened for reading: %s"), *Stream.LogFilePath);
+			continue;
+		}
+
+		const int64 CurrentFileSize = FileHandle->Size();
+		if (CurrentFileSize <= Stream.LastFilePosition)
+		{
+			// Nothing to send.
+			continue;
+		}
+
+		// Calculate how much new content to read (limit to 32KB per update)
+
+		constexpr int64 MaxReadSize = 32768;
+		const int64 AvailableBytes = CurrentFileSize - Stream.LastFilePosition;
+		const int64 ReadSize = FMath::Min(AvailableBytes, MaxReadSize);
+
+		TArray<uint8> Buffer;
+		Buffer.SetNumUninitialized(ReadSize);
+
+		if (FileHandle->Seek(Stream.LastFilePosition) && FileHandle->Read(Buffer.GetData(), ReadSize))
+		{
+			Stream.LastFilePosition += ReadSize;
+
+			TSharedPtr<FJsonObject> JsonObject = MakeShareable(new FJsonObject);
+			JsonObject->SetStringField(TEXT("command"), TEXT("live log data"));
+			JsonObject->SetStringField(TEXT("log_file_path"), Stream.LogFilePath);
+			JsonObject->SetNumberField(TEXT("file_position"), Stream.LastFilePosition - ReadSize);
+			JsonObject->SetNumberField(TEXT("content_size"), ReadSize);
+
+			// Compress if enabled and chunk is large enough to benefit from compression
+
+			const int64 CompressionThreshold = 1024; // 1KB threshold
+
+			if (ReadSize >= CompressionThreshold)
+			{
+				// Try to compress the data
+				const int32 CompressedUpperBound = FCompression::CompressMemoryBound(NAME_Zlib, ReadSize);
+				TArray<uint8> CompressedData;
+				CompressedData.SetNumUninitialized(CompressedUpperBound);
+
+				int32 CompressedSize = CompressedUpperBound;
+				if (FCompression::CompressMemory(NAME_Zlib, CompressedData.GetData(), CompressedSize, Buffer.GetData(), ReadSize, COMPRESS_BiasSpeed))
+				{
+					// Only use compression if it actually saves space
+					if (CompressedSize < ReadSize)
+					{
+						CompressedData.SetNum(CompressedSize);
+						FString CompressedBase64Content = FBase64::Encode(CompressedData.GetData(), CompressedSize);
+						JsonObject->SetStringField(TEXT("content_compressed_base64"), CompressedBase64Content);
+						JsonObject->SetNumberField(TEXT("uncompressed_size"), ReadSize);
+					}
+					else
+					{
+						// Compression didn't help, use uncompressed
+						FString Base64Content = FBase64::Encode(Buffer.GetData(), ReadSize);
+						JsonObject->SetStringField(TEXT("content_base64"), Base64Content);
+					}
+				}
+				else
+				{
+					// Compression failed, use uncompressed
+					FString Base64Content = FBase64::Encode(Buffer.GetData(), ReadSize);
+					JsonObject->SetStringField(TEXT("content_base64"), Base64Content);
+				}
+			}
+			else
+			{
+				// No compression (either disabled or chunk too small)
+				FString Base64Content = FBase64::Encode(Buffer.GetData(), ReadSize);
+				JsonObject->SetStringField(TEXT("content_base64"), Base64Content);
+			}
+
+			FString OutputString;
+			TSharedRef<TJsonWriter<>> Writer = TJsonWriterFactory<>::Create(&OutputString);
+			FJsonSerializer::Serialize(JsonObject.ToSharedRef(), Writer);
+
+			SendMessage(OutputString, Stream.ClientEndpoint);
+		}
+	}
+}
+
 void FSwitchboardListener::CleanUpDisconnectedSockets()
 {
 	SWITCHBOARD_TRACE_CPUPROFILER_EVENT_SCOPE(FSwitchboardListener::CleanUpDisconnectedSockets);
diff --git a/Engine/Source/Programs/SwitchboardListener/SblCore/Private/SwitchboardProtocol.cpp b/Engine/Source/Programs/SwitchboardListener/SblCore/Private/SwitchboardProtocol.cpp
--- a/Engine/Source/Programs/SwitchboardListener/SblCore/Private/SwitchboardProtocol.cpp
+++ b/Engine/Source/Programs/SwitchboardListener/SblCore/Private/SwitchboardProtocol.cpp
@@ -121,6 +121,13 @@ FString CreateRedeployStatusMessage(const FGuid& InMessageID, bool bAck, const F
 	return CreateMessage(TEXT("redeploy server status"), bAck, { { TEXT("id"), InMessageID.ToString() }, { TEXT("status"), Status } });
 }
 
+FString CreateLiveLogStreamTerminatedMessage(const FString& InLogFilePath, const FString& InReason)
+{
+	return CreateMessage(TEXT("live log stream terminated"), true, {
+		{ TEXT("log_file_path"), InLogFilePath },
+		{ TEXT("reason"), InReason }
+	});
+}
 
 // This logs an error with some context in the event the field is missing.
 TSharedPtr<FJsonValue> TryGetCommandRequiredField(const TSharedPtr<FJsonObject>& CommandObj, const FString& FieldName)
@@ -437,6 +444,48 @@ FCreateTaskResult CreateTaskFromCommand(const FString& InCommand, const FIPv4End
 		Result.Task = MakeUnique<FSwitchboardMinimizeWindowsTask>(MessageID, InEndpoint);
 		return Result;
 	}
+	else if (CommandName == FSwitchboardStartLiveLogStreamTask::CommandName)
+	{
+		TSharedPtr<FJsonValue> LogFilePathField = TryGetCommandRequiredField(JsonData, TEXT("log_file_path"));
+
+		if (LogFilePathField.IsValid())
+		{
+			// Optional tail_only field
+
+			bool bTailOnly = false; // Default to send from the beginning
+
+			TSharedPtr<FJsonValue> TailOnlyField = JsonData->TryGetField(TEXT("tail_only"));
+
+			if (TailOnlyField.IsValid())
+			{
+				bTailOnly = TailOnlyField->AsBool();
+			}
+
+			// Optional process_uuid field (currently there always is one)
+
+			FGuid ProcessUUID;
+			TSharedPtr<FJsonValue> ProcessUUIDField = JsonData->TryGetField(TEXT("process_uuid"));
+			if (ProcessUUIDField.IsValid())
+			{
+				FGuid::Parse(ProcessUUIDField->AsString(), ProcessUUID);
+			}
+
+			Result.Status = ECreateTaskStatus::Success;
+			Result.Task = MakeUnique<FSwitchboardStartLiveLogStreamTask>(MessageID, InEndpoint, LogFilePathField->AsString(), bTailOnly, ProcessUUID);
+			return Result;
+		}
+	}
+	else if (CommandName == FSwitchboardStopLiveLogStreamTask::CommandName)
+	{
+		TSharedPtr<FJsonValue> LogFilePathField = TryGetCommandRequiredField(JsonData, TEXT("log_file_path"));
+
+		if (LogFilePathField.IsValid())
+		{
+			Result.Status = ECreateTaskStatus::Success;
+			Result.Task = MakeUnique<FSwitchboardStopLiveLogStreamTask>(MessageID, InEndpoint, LogFilePathField->AsString());
+			return Result;
+		}
+	}
 
 	Result.Status = ECreateTaskStatus::Error_Unhandled;
 	return Result;
diff --git a/Engine/Source/Programs/SwitchboardListener/SblCore/Private/SwitchboardProtocol.h b/Engine/Source/Programs/SwitchboardListener/SblCore/Private/SwitchboardProtocol.h
--- a/Engine/Source/Programs/SwitchboardListener/SblCore/Private/SwitchboardProtocol.h
+++ b/Engine/Source/Programs/SwitchboardListener/SblCore/Private/SwitchboardProtocol.h
@@ -41,6 +41,8 @@ FString CreateSyncStatusMessage(const FSyncStatus& SyncStatus, ESyncStatusReques
 
 FString CreateRedeployStatusMessage(const FGuid& InMessageID, bool bAck, const FString& Status);
 
+FString CreateLiveLogStreamTerminatedMessage(const FString& InLogFilePath, const FString& InReason);
+
 //~
 
 enum class ECreateTaskStatus
diff --git a/Engine/Source/Programs/SwitchboardListener/SblCore/Private/SwitchboardTasks.h b/Engine/Source/Programs/SwitchboardListener/SblCore/Private/SwitchboardTasks.h
--- a/Engine/Source/Programs/SwitchboardListener/SblCore/Private/SwitchboardTasks.h
+++ b/Engine/Source/Programs/SwitchboardListener/SblCore/Private/SwitchboardTasks.h
@@ -23,7 +23,9 @@ enum class ESwitchboardTaskType : uint8
 	RefreshMosaics,
 	MinimizeWindows,
 	SetInactiveTimeout,
-	FreeListenerBinary
+	FreeListenerBinary,
+	StartLiveLogStream,
+	StopLiveLogStream
 };
 
 
@@ -366,3 +368,49 @@ struct FSwitchboardFreeListenerBinaryTask : public FSwitchboardTask
 	virtual const TCHAR* GetCommandName() const override { return CommandName; };
 	//~ End FSwitchboardTask interface
 };
+
+struct FSwitchboardStartLiveLogStreamTask : public FSwitchboardTask
+{
+	FSwitchboardStartLiveLogStreamTask(const FGuid& InTaskID, const FIPv4Endpoint& InEndpoint, const FString& InLogFilePath, bool bInTailOnly = false, const FGuid& InProcessUUID = FGuid())
+		: FSwitchboardTask{ ESwitchboardTaskType::StartLiveLogStream, InTaskID, InEndpoint }
+		, LogFilePath(InLogFilePath)
+		, bTailOnly(bInTailOnly)
+		, ProcessUUID(InProcessUUID)
+		
+	{}
+
+	FString LogFilePath;
+	bool bTailOnly;
+	FGuid ProcessUUID;
+
+	//~ Begin FSwitchboardTask interface
+	static constexpr const TCHAR* CommandName = TEXT("start live log stream");
+	virtual const TCHAR* GetCommandName() const override { return CommandName; };
+
+	virtual uint32 GetEquivalenceHash() const override
+	{
+		uint32 Hash = HashCombine(FSwitchboardTask::GetEquivalenceHash(), GetTypeHash(LogFilePath));
+		return HashCombine(Hash, GetTypeHash(bTailOnly));
+	}
+	//~ End FSwitchboardTask interface
+};
+
+struct FSwitchboardStopLiveLogStreamTask : public FSwitchboardTask
+{
+	FSwitchboardStopLiveLogStreamTask(const FGuid& InTaskID, const FIPv4Endpoint& InEndpoint, const FString& InLogFilePath)
+		: FSwitchboardTask{ ESwitchboardTaskType::StopLiveLogStream, InTaskID, InEndpoint }
+		, LogFilePath(InLogFilePath)
+	{}
+
+	FString LogFilePath;
+
+	//~ Begin FSwitchboardTask interface
+	static constexpr const TCHAR* CommandName = TEXT("stop live log stream");
+	virtual const TCHAR* GetCommandName() const override { return CommandName; };
+
+	virtual uint32 GetEquivalenceHash() const override
+	{
+		return HashCombine(FSwitchboardTask::GetEquivalenceHash(), GetTypeHash(LogFilePath));
+	}
+	//~ End FSwitchboardTask interface
+};
