Switchboard: Fix not (re-) populating the nD Tools dropdowns when a config loads.

(Cherrypick of CL 46078032.)

#rnx

#Synced-CL 48175258

diff --git a/Engine/Plugins/VirtualProduction/Switchboard/Source/Switchboard/switchboard/devices/ndisplay/ndisplay_tools_ui.py b/Engine/Plugins/VirtualProduction/Switchboard/Source/Switchboard/switchboard/devices/ndisplay/ndisplay_tools_ui.py
--- a/Engine/Plugins/VirtualProduction/Switchboard/Source/Switchboard/switchboard/devices/ndisplay/ndisplay_tools_ui.py
+++ b/Engine/Plugins/VirtualProduction/Switchboard/Source/Switchboard/switchboard/devices/ndisplay/ndisplay_tools_ui.py
@@ -80,8 +80,12 @@ class nDisplayToolsUI(QWidget):
 
     def on_config_changed(self, old_value, new_value):
         '''Called when nDisplay config file changes - update default selections'''
-        # When config changes, the default media profile and livelink preset may have changed
-        # So we need to update our dropdown selections
+        # When config changes, refresh dropdown contents from SBCache first
+        # as the loaded config might be from a different project with different assets
+        self.refresh_media_profiles()
+        self.refresh_livelink_presets()
+
+        # Then update dropdown selections to match defaults
         self._set_default_media_profile()
         self._set_default_livelink_preset()
 
