Switchboard: Added a custom commands section to the ndisplay tools panel.

Very convenient for triggering "ke" blueprint custom events, faster than setting up remote control.

(Cherrypick of CL 46295681.)

#rnx

#Synced-CL 48175258

diff --git a/Engine/Plugins/VirtualProduction/Switchboard/Source/Switchboard/switchboard/devices/ndisplay/ndisplay_tools_ui.py b/Engine/Plugins/VirtualProduction/Switchboard/Source/Switchboard/switchboard/devices/ndisplay/ndisplay_tools_ui.py
--- a/Engine/Plugins/VirtualProduction/Switchboard/Source/Switchboard/switchboard/devices/ndisplay/ndisplay_tools_ui.py
+++ b/Engine/Plugins/VirtualProduction/Switchboard/Source/Switchboard/switchboard/devices/ndisplay/ndisplay_tools_ui.py
@@ -1,9 +1,12 @@
 # Copyright Epic Games, Inc. All Rights Reserved.
 
+import json
 from PySide6 import QtGui, QtWidgets
+from PySide6.QtCore import Qt
 from PySide6.QtWidgets import (
-    QHBoxLayout, QPushButton,
-    QSizePolicy, QVBoxLayout, QWidget, QGroupBox
+    QDialog, QDialogButtonBox, QFormLayout, QGridLayout,
+    QHBoxLayout, QLabel, QLineEdit, QMessageBox, QPushButton,
+    QSizePolicy, QToolButton, QVBoxLayout, QWidget, QGroupBox
 )
 
 from switchboard import switchboard_widgets as sb_widgets
@@ -13,12 +16,67 @@ from switchboard.switchboard_logging import LOGGER
 from switchboard.devices.unreal.plugin_unreal import DeviceUnreal
 
 
+class CustomCommandDialog(QDialog):
+    '''Dialog for adding/editing custom exec commands'''
+
+    def __init__(self, parent, name='', command=''):
+        super().__init__(parent)
+
+        self.setWindowTitle('Custom Command' if not name else 'Edit Custom Command')
+        self.setMinimumWidth(400)
+
+        # Main layout
+        layout = QVBoxLayout(self)
+
+        # Form layout
+        form_layout = QFormLayout()
+
+        # Name input
+        self.name_edit = QLineEdit(name)
+        self.name_edit.setPlaceholderText('e.g. Trigger Event')
+        form_layout.addRow('Name:', self.name_edit)
+
+        # Command input
+        self.command_edit = QLineEdit(command)
+        self.command_edit.setPlaceholderText('e.g. ke * MyBlueprintCustomEventName')
+        self.command_edit.setToolTip('Console command to execute on the nDisplay cluster')
+        form_layout.addRow('Command:', self.command_edit)
+
+        layout.addLayout(form_layout)
+
+        # Buttons
+        self.button_box = QDialogButtonBox(
+            QDialogButtonBox.Ok | QDialogButtonBox.Cancel,
+            parent=self
+        )
+        self.button_box.accepted.connect(self.validate_and_accept)
+        self.button_box.rejected.connect(self.reject)
+        layout.addWidget(self.button_box)
+
+    def validate_and_accept(self):
+        '''Validates input and accepts dialog'''
+        if not self.name_edit.text().strip():
+            QMessageBox.warning(self, 'Invalid Input', 'Please enter a name to use for the button label.')
+            return
+        if not self.command_edit.text().strip():
+            QMessageBox.warning(self, 'Invalid Input', 'Please enter a command to execute.')
+            return
+        self.accept()
+
+    def get_values(self):
+        '''Returns the entered name and command'''
+        return self.name_edit.text().strip(), self.command_edit.text().strip()
+
+
 class nDisplayToolsUI(QWidget):
     '''
     UI for nDisplay Tools tab that allows setting Media Profiles and
     applying Live Link presets to the nDisplay cluster at runtime.
     '''
 
+    # Configuration
+    CUSTOM_COMMANDS_NUM_COLUMNS = 4
+
     def __init__(self, parent: QWidget, plugin_ndisplay):
         QWidget.__init__(self, parent)
 
@@ -36,9 +94,9 @@ class nDisplayToolsUI(QWidget):
         livelink_preset_group = self.create_livelink_preset_section()
         layout.addWidget(livelink_preset_group)
 
-        # Timecode section
-        timecode_group = self.create_timecode_section()
-        layout.addWidget(timecode_group)
+        # Custom Commands section
+        custom_commands_group = self.create_custom_commands_section()
+        layout.addWidget(custom_commands_group)
 
         # Add stretch to push everything to the top
         layout.addStretch()
@@ -65,6 +123,10 @@ class nDisplayToolsUI(QWidget):
         self.plugin_ndisplay.csettings['livelink_preset'].signal_setting_changed.connect(
             self.on_default_settings_changed)
 
+        # Listen for custom commands changes
+        self.plugin_ndisplay.csettings['custom_exec_commands'].signal_setting_changed.connect(
+            self.on_custom_commands_changed)
+
         # Initial population of combo boxes
         self.refresh_media_profiles()
         self.refresh_livelink_presets()
@@ -95,6 +157,10 @@ class nDisplayToolsUI(QWidget):
         self._set_default_media_profile()
         self._set_default_livelink_preset()
 
+    def on_custom_commands_changed(self, old_value, new_value):
+        '''Called when custom commands settings change'''
+        self.refresh_custom_commands()
+
     def _set_default_media_profile(self):
         '''Sets the media profile combo box to the default value from nDisplay device settings'''
         # Get default media profile from nDisplay class settings
@@ -201,27 +267,47 @@ class nDisplayToolsUI(QWidget):
 
         return group_box
 
-    def create_timecode_section(self) -> QGroupBox:
-        '''Creates the Timecode section of the UI'''
-        group_box = QGroupBox("Timecode")
-        layout = QVBoxLayout()
+    def create_custom_commands_section(self) -> QGroupBox:
+        '''Creates the Custom Commands section of the UI'''
 
-        # Buttons layout
-        controls_layout = QHBoxLayout()
+        group_box = QGroupBox("Custom Commands")
+        main_layout = QVBoxLayout()
 
-        # Reinitialize button
-        self.timecode_reinitialize_btn = QPushButton("Reinitialize Provider")
-        self.timecode_reinitialize_btn.setToolTip(
-            "Reinitialize the current timecode provider on the nDisplay cluster")
-        self.timecode_reinitialize_btn.clicked.connect(
-            self.on_reinitialize_timecode_provider)
-        controls_layout.addWidget(self.timecode_reinitialize_btn)
+        # Container widget for command buttons
+        self.commands_container = QWidget()
+        self.commands_grid = QGridLayout(self.commands_container)
+        self.commands_grid.setSpacing(5)
 
-        # Add stretch to push button to the left
-        controls_layout.addStretch()
+        # Set uniform column widths
+        for col in range(self.CUSTOM_COMMANDS_NUM_COLUMNS):
+            self.commands_grid.setColumnStretch(col, 1)
 
-        layout.addLayout(controls_layout)
-        group_box.setLayout(layout)
+        main_layout.addWidget(self.commands_container)
+
+        # Empty state label
+        self.empty_commands_label = QLabel("No custom commands added yet")
+        self.empty_commands_label.setStyleSheet("QLabel { color: gray; }")
+        self.empty_commands_label.setAlignment(Qt.AlignCenter)
+        main_layout.addWidget(self.empty_commands_label)
+
+        # Add command button at bottom left
+        bottom_layout = QHBoxLayout()
+        self.add_command_btn = QPushButton()
+        self.add_command_btn.setIcon(
+            QtGui.QIcon(':/icons/images/PlusSymbol_12x.png'))
+        self.add_command_btn.setProperty(u"frameless", True)
+        self.add_command_btn.setMaximumWidth(30)
+        self.add_command_btn.setToolTip("Add a new custom command")
+        self.add_command_btn.clicked.connect(self.add_custom_command)
+        bottom_layout.addWidget(self.add_command_btn)
+        bottom_layout.addStretch()
+
+        main_layout.addLayout(bottom_layout)
+
+        group_box.setLayout(main_layout)
+
+        # Load existing commands
+        self.refresh_custom_commands()
 
         return group_box
 
@@ -368,16 +454,197 @@ class nDisplayToolsUI(QWidget):
             except Exception as e:
                 LOGGER.error(f"Failed to clear Live Link sources: {e}")
 
-    def on_reinitialize_timecode_provider(self):
-        '''Reinitializes the timecode provider on the nDisplay cluster'''
+    def add_custom_command(self):
+        '''Shows dialog to add a new custom command'''
+
+        dialog = CustomCommandDialog(self)
+        if dialog.exec_() == QDialog.Accepted:
+            name, command = dialog.get_values()
+
+            # Get existing commands
+            commands = self._get_custom_commands()
+
+            # Add new command
+            commands.append({
+                'name': name,
+                'command': command
+            })
+
+            # Save to settings
+            self._save_custom_commands(commands)
+
+            # Refresh UI
+            self.refresh_custom_commands()
+
+    def edit_custom_command(self, index):
+        '''Shows dialog to edit an existing custom command'''
+
+        commands = self._get_custom_commands()
+
+        if 0 <= index < len(commands):
+            cmd = commands[index]
+            dialog = CustomCommandDialog(self, cmd['name'], cmd['command'])
+            if dialog.exec_() == QDialog.Accepted:
+                name, command = dialog.get_values()
+                commands[index] = {'name': name, 'command': command}
+                self._save_custom_commands(commands)
+                self.refresh_custom_commands()
+
+    def remove_custom_command(self, index):
+        '''Removes a custom command after confirmation'''
+
+        commands = self._get_custom_commands()
+
+        if 0 <= index < len(commands):
+            cmd = commands[index]
+            reply = QMessageBox.question(
+                self,
+                'Remove Custom Command',
+                f"Remove command '{cmd['name']}'?",
+                QMessageBox.Yes | QMessageBox.No,
+                QMessageBox.No
+            )
+            if reply == QMessageBox.Yes:
+                commands.pop(index)
+                self._save_custom_commands(commands)
+                self.refresh_custom_commands()
+
+    def move_custom_command(self, from_index, to_index):
+        '''Moves a custom command from one position to another'''
+
+        commands = self._get_custom_commands()
+
+        if (0 <= from_index < len(commands) and
+            0 <= to_index < len(commands) and
+            from_index != to_index):
+
+            # Move the command
+            command = commands.pop(from_index)
+            commands.insert(to_index, command)
+
+            # Save and refresh
+            self._save_custom_commands(commands)
+            self.refresh_custom_commands()
+
+    def execute_custom_command(self, command):
+        '''Executes a custom command on the nDisplay cluster'''
+
         devices = self.get_connected_devices()
+
         if not devices:
             LOGGER.warning("No connected nDisplay devices found")
             return
 
-        exec_str = "TimecodeProvider.Reinitialize"
         try:
-            devices[0].console_exec_cluster(devices, exec_str)
-            LOGGER.info("Reinitializing timecode provider")
+            devices[0].console_exec_cluster(devices, command)
+            LOGGER.info(f"Executed custom command: {command}")
         except Exception as e:
-            LOGGER.error(f"Failed to reinitialize timecode provider: {e}")
+            LOGGER.error(f"Failed to execute custom command: {e}")
+
+    def refresh_custom_commands(self):
+        '''Refreshes the custom command buttons from settings'''
+
+        # Clear existing buttons
+        while self.commands_grid.count():
+            item = self.commands_grid.takeAt(0)
+            if item.widget():
+                item.widget().deleteLater()
+
+        # Get commands from settings
+        commands = self._get_custom_commands()
+
+        # Show/hide empty state
+        self.empty_commands_label.setVisible(len(commands) == 0)
+        self.commands_container.setVisible(len(commands) > 0)
+
+        # Create buttons in grid
+        columns = self.CUSTOM_COMMANDS_NUM_COLUMNS
+        for i, cmd in enumerate(commands):
+            row = i // columns
+            col = i % columns
+
+            # Create down arrow button.
+            button_widget = QWidget()
+            button_widget.setSizePolicy(QSizePolicy.Expanding, QSizePolicy.Fixed)
+            button_layout = QHBoxLayout(button_widget)
+            button_layout.setContentsMargins(0, 0, 0, 0)
+            button_layout.setSpacing(2)
+
+            # Main command button
+            cmd_btn = QPushButton()
+            cmd_btn.setMinimumHeight(50) # Intended to look taller than usual, about the width of device rows.
+            cmd_btn.setToolTip(f"Execute: {cmd['command']}")
+            cmd_btn.setSizePolicy(QSizePolicy.Expanding, QSizePolicy.Fixed)
+            cmd_btn.clicked.connect(lambda checked=False, c=cmd['command']: self.execute_custom_command(c))
+
+            # Create a label with word wrap inside the button (also helps keep the button width distribution even)
+            btn_label = QLabel(cmd['name'])
+            btn_label.setWordWrap(True)
+            btn_label.setAlignment(Qt.AlignCenter)
+            btn_label.setStyleSheet("QLabel { background: transparent; }")
+
+            # Set the label as the button's layout
+            btn_layout = QVBoxLayout(cmd_btn)
+            btn_layout.setContentsMargins(5, 5, 5, 5)
+            btn_layout.addWidget(btn_label)
+
+            button_layout.addWidget(cmd_btn)
+
+            # Options button
+            options_btn = QToolButton()
+            options_btn.setIcon(QtGui.QIcon(':/icons/images/down_arrow.png'))
+            options_btn.setProperty(u"frameless", True)
+            options_btn.setFixedWidth(30)
+            options_btn.setSizePolicy(QSizePolicy.Fixed, QSizePolicy.Fixed)
+            options_btn.setToolTip("Edit or remove this command")
+            options_btn.clicked.connect(lambda checked=False, idx=i, btn=options_btn: self._show_command_menu(idx, btn))
+
+            button_layout.addWidget(options_btn, 0, Qt.AlignTop)
+
+            self.commands_grid.addWidget(button_widget, row, col)
+
+    def _show_command_menu(self, index, button):
+        '''Shows context menu for command editing/removal'''
+
+        menu = QtWidgets.QMenu(self)
+        commands = self._get_custom_commands()
+
+        edit_action = menu.addAction(QtGui.QIcon(':/icons/images/icon_edit.png'), "Edit")
+        edit_action.triggered.connect(lambda: self.edit_custom_command(index))
+
+        menu.addSeparator()
+
+        # Move back option (disabled if already at beginning)
+        move_back_action = menu.addAction("Move Back")
+        move_back_action.setEnabled(index > 0)
+        move_back_action.triggered.connect(lambda: self.move_custom_command(index, index - 1))
+
+        # Move forward option (disabled if already at end)
+        move_forward_action = menu.addAction("Move Forward")
+        move_forward_action.setEnabled(index < len(commands) - 1)
+        move_forward_action.triggered.connect(lambda: self.move_custom_command(index, index + 1))
+
+        menu.addSeparator()
+
+        remove_action = menu.addAction(QtGui.QIcon(':/icons/images/icon_remove.png'), "Remove")
+        remove_action.triggered.connect(lambda: self.remove_custom_command(index))
+
+        menu.exec_(button.mapToGlobal(button.rect().bottomLeft()))
+
+    def _get_custom_commands(self):
+        '''Gets custom commands from settings'''
+
+        commands_json = self.plugin_ndisplay.csettings['custom_exec_commands'].get_value()
+        commands = []
+        for cmd_str in commands_json:
+            try:
+                commands.append(json.loads(cmd_str))
+            except json.JSONDecodeError:
+                LOGGER.warning(f"Invalid custom command JSON: {cmd_str}")
+        return commands
+
+    def _save_custom_commands(self, commands):
+        '''Saves custom commands to settings'''
+
+        commands_json = [json.dumps(cmd) for cmd in commands]
+        self.plugin_ndisplay.csettings['custom_exec_commands'].update_value(commands_json)
diff --git a/Engine/Plugins/VirtualProduction/Switchboard/Source/Switchboard/switchboard/devices/ndisplay/plugin_ndisplay.py b/Engine/Plugins/VirtualProduction/Switchboard/Source/Switchboard/switchboard/devices/ndisplay/plugin_ndisplay.py
--- a/Engine/Plugins/VirtualProduction/Switchboard/Source/Switchboard/switchboard/devices/ndisplay/plugin_ndisplay.py
+++ b/Engine/Plugins/VirtualProduction/Switchboard/Source/Switchboard/switchboard/devices/ndisplay/plugin_ndisplay.py
@@ -718,6 +718,13 @@ class DevicenDisplay(DeviceUnreal):
             nice_name='Cluster exec console command history',
             value=[],
             show_ui=False,
+        ),
+        'custom_exec_commands': StringListSetting(
+            attr_name='custom_exec_commands',
+            nice_name='Custom Exec Commands',
+            value=[],
+            show_ui=False,
+            tool_tip='List of custom exec commands for the nDisplay Tools UI'
         )
     }
 
