Switchboard: Fix not updating custom commands when another config is loaded.

(Cherrypick of CL 46295825.)

#rnx

#Synced-CL 48175258

diff --git a/Engine/Plugins/VirtualProduction/Switchboard/Source/Switchboard/switchboard/devices/ndisplay/ndisplay_tools_ui.py b/Engine/Plugins/VirtualProduction/Switchboard/Source/Switchboard/switchboard/devices/ndisplay/ndisplay_tools_ui.py
--- a/Engine/Plugins/VirtualProduction/Switchboard/Source/Switchboard/switchboard/devices/ndisplay/ndisplay_tools_ui.py
+++ b/Engine/Plugins/VirtualProduction/Switchboard/Source/Switchboard/switchboard/devices/ndisplay/ndisplay_tools_ui.py
@@ -146,6 +146,7 @@ class nDisplayToolsUI(QWidget):
         # as the loaded config might be from a different project with different assets
         self.refresh_media_profiles()
         self.refresh_livelink_presets()
+        self.refresh_custom_commands()
 
         # Then update dropdown selections to match defaults
         self._set_default_media_profile()
