Switchboard: Re-opening live logs starts from last offset instead of from scratch.

This makes re-opening the same live log much faster, especially when the live log file is large.

(Cherrypick of CL 46295859.)

#rnx

#Synced-CL 48191591

diff --git a/Engine/Plugins/VirtualProduction/Switchboard/Source/Switchboard/switchboard/devices/unreal/plugin_unreal.py b/Engine/Plugins/VirtualProduction/Switchboard/Source/Switchboard/switchboard/devices/unreal/plugin_unreal.py
--- a/Engine/Plugins/VirtualProduction/Switchboard/Source/Switchboard/switchboard/devices/unreal/plugin_unreal.py
+++ b/Engine/Plugins/VirtualProduction/Switchboard/Source/Switchboard/switchboard/devices/unreal/plugin_unreal.py
@@ -1192,6 +1192,9 @@ class DeviceUnreal(Device):
         self.live_log_file_path = None
         self.live_log_local_file = None
 
+        # Track log files per process UUID for resuming
+        self.process_log_files = {}  # UUID -> {'local_file': pathlib.Path, 'last_position': int}
+
         self.osc_connection_timer = QtCore.QTimer(self)
         self.osc_connection_timer.timeout.connect(self._try_connect_osc)
 
@@ -1720,6 +1723,11 @@ class DeviceUnreal(Device):
                 self, 'disconnect_listener', QtCore.Qt.QueuedConnection)
             return
 
+        # Clear process log file tracking on disconnect
+        if self.process_log_files:
+            LOGGER.info(f"Clearing log file tracking for {len(self.process_log_files)} processes on disconnect")
+            self.process_log_files.clear()
+
         super().disconnect_listener()
         self.unreal_client.disconnect()
 
@@ -2804,6 +2812,13 @@ class DeviceUnreal(Device):
                 f"{self.name}: on_program_ended with unknown id {puuid}")
             return
 
+        # Clean up log file tracking for this process if it's an Unreal process
+        if program_name == UnrealJobs.Unreal.value:
+            process_uuid_str = str(puuid)
+            if process_uuid_str in self.process_log_files:
+                LOGGER.info(f"Removing log file tracking for ended Unreal process {process_uuid_str}")
+                del self.process_log_files[process_uuid_str]
+
         self.do_program_ended_update(
             program_name=program_name,
             returncode=returncode,
@@ -3521,24 +3536,6 @@ class DeviceUnreal(Device):
             LOGGER.error(f"Cannot start live log because Unreal is not running on device {self.name}")
             return
 
-        # Construct the remote log file path
-        remote_log_path = self.get_remote_log_path() / self.log_filename
-        LOGGER.info(f"Live log: requesting streaming for remote path: {remote_log_path}")
-
-        # Create local live log file
-        log_dir = DeviceUnreal.get_log_download_dir()
-        timestamp = datetime.now().strftime("%Y%m%d_%H%M%S")
-        local_filename = f"live_{self.name}_{timestamp}.log"
-        self.live_log_local_file = pathlib.Path(log_dir) / local_filename
-
-        # Create the local live log file, that the log viewer will display as it updates.
-        try:
-            self.live_log_local_file.parent.mkdir(parents=True, exist_ok=True)
-            self.live_log_local_file.touch()
-        except Exception as exception:
-            LOGGER.error(f"Failed to create live log file {self.live_log_local_file}: {exception}")
-            return
-
         # Get the UUID of the running Unreal process (if any)
 
         unreal_process_uuid = None
@@ -3551,10 +3548,51 @@ class DeviceUnreal(Device):
             LOGGER.error("Live log: no running Unreal processes found, can't start the live log")
             return
 
+        # Construct the remote log file path
+        remote_log_path = self.get_remote_log_path() / self.log_filename
+        LOGGER.info(f"Live log: requesting streaming for remote path: {remote_log_path}")
+
+        # Check if we have an existing log file for this process UUID
+        start_offset = None
+        if unreal_process_uuid in self.process_log_files:
+            existing_log = self.process_log_files[unreal_process_uuid]
+            if existing_log['local_file'].exists():
+                # Reuse the existing log file and resume from where we left off
+                self.live_log_local_file = existing_log['local_file']
+                start_offset = existing_log['local_file'].stat().st_size
+                LOGGER.info(f"Resuming live log from existing file {self.live_log_local_file}, offset: {start_offset} bytes")
+            else:
+                # File was deleted, remove from tracking
+                del self.process_log_files[unreal_process_uuid]
+                LOGGER.info(f"Previous log file was deleted, will create new one")
+
+        # Create new local log file if we don't have one
+        if self.live_log_local_file is None or not self.live_log_local_file.exists():
+            log_dir = DeviceUnreal.get_log_download_dir()
+            timestamp = datetime.now().strftime("%Y%m%d_%H%M%S")
+            local_filename = f"live_{self.name}_{timestamp}.log"
+            self.live_log_local_file = pathlib.Path(log_dir) / local_filename
+
+            # Create the local live log file
+            try:
+                self.live_log_local_file.parent.mkdir(parents=True, exist_ok=True)
+                self.live_log_local_file.touch()
+                start_offset = 0  # New file, start from beginning
+                LOGGER.info(f"Created new live log file: {self.live_log_local_file}")
+            except Exception as exception:
+                LOGGER.error(f"Failed to create live log file {self.live_log_local_file}: {exception}")
+                return
+
+            # Track this new log file for the process
+            self.process_log_files[unreal_process_uuid] = {
+                'local_file': self.live_log_local_file,
+                'last_position': 0
+            }
+
         # Send start live log message to listener
 
         msgid, msg = message_protocol.create_start_live_log_stream_message(
-            str(remote_log_path), tail_only=False, process_uuid=unreal_process_uuid)
+            str(remote_log_path), tail_only=False, process_uuid=unreal_process_uuid, start_offset=start_offset)
 
         try:
             self.unreal_client.send_message(msg)
@@ -3579,6 +3617,14 @@ class DeviceUnreal(Device):
             LOGGER.warning(f"Received live log data but streaming not active or no local file: streaming={self.live_log_streaming}, file={self.live_log_local_file}")
             return
 
+        # Update tracked position for this process
+        unreal_puuids = self.program_start_queue.running_puuids_named(UnrealJobs.Unreal.value)
+        if unreal_puuids:
+            process_uuid = str(unreal_puuids[0])
+            if process_uuid in self.process_log_files:
+                # Track the file position (could also use file_position from message)
+                self.process_log_files[process_uuid]['last_position'] = self.live_log_local_file.stat().st_size
+
         try:
             # Check for compressed content first
 
@@ -3625,6 +3671,9 @@ class DeviceUnreal(Device):
         if not self.live_log_streaming:
             return
 
+        # Note: We intentionally do NOT clear process_log_files here
+        # so that we can resume if the same process is still running
+
         if self.live_log_file_path and self.unreal_client.is_connected:
 
             # Send stop live log message to listener
@@ -3662,6 +3711,9 @@ class DeviceUnreal(Device):
         if self.live_log_streaming:
             self._stop_live_log_internal()
 
+        # Note: We keep the process_log_files tracking so we can resume from the same position
+        # if the user reopens the log viewer while the process is still running
+
     def on_copy_last_launch_command(self):
         ''' Copies the last launch command to the clipboard'''
         QtGui.QGuiApplication.clipboard().setText(
diff --git a/Engine/Plugins/VirtualProduction/Switchboard/Source/Switchboard/switchboard/devices/unreal/version_helpers.py b/Engine/Plugins/VirtualProduction/Switchboard/Source/Switchboard/switchboard/devices/unreal/version_helpers.py
--- a/Engine/Plugins/VirtualProduction/Switchboard/Source/Switchboard/switchboard/devices/unreal/version_helpers.py
+++ b/Engine/Plugins/VirtualProduction/Switchboard/Source/Switchboard/switchboard/devices/unreal/version_helpers.py
@@ -3,7 +3,7 @@
 import typing
 
 
-LISTENER_COMPATIBLE_VERSION = (3,4)
+LISTENER_COMPATIBLE_VERSION = (3,5)
 
 # Version where redeploy support was initially introduced.
 LISTENER_MIN_REDEPLOY_VERSION = (1,4,0)
diff --git a/Engine/Plugins/VirtualProduction/Switchboard/Source/Switchboard/switchboard/message_protocol.py b/Engine/Plugins/VirtualProduction/Switchboard/Source/Switchboard/switchboard/message_protocol.py
--- a/Engine/Plugins/VirtualProduction/Switchboard/Source/Switchboard/switchboard/message_protocol.py
+++ b/Engine/Plugins/VirtualProduction/Switchboard/Source/Switchboard/switchboard/message_protocol.py
@@ -233,7 +233,7 @@ def create_set_inactive_timeout_message(timeout_seconds: int):
     return (cmd_id, message)
 
 
-def create_start_live_log_stream_message(log_file_path: str, tail_only: bool = False, process_uuid: str = None):
+def create_start_live_log_stream_message(log_file_path: str, tail_only: bool = False, process_uuid: str = None, start_offset: Optional[int] = None):
     ''' Creates the message to start SBL streaming the live log of the running Unreal/nDisplay process '''
 
     cmd_id = uuid.uuid4()
@@ -245,6 +245,10 @@ def create_start_live_log_stream_message(log_file_path: str, tail_only: bool = F
     }
     if process_uuid:
         cmd['process_uuid'] = process_uuid
+
+    if start_offset is not None and start_offset >= 0:
+        cmd['start_offset'] = start_offset
+
     message = json.dumps(cmd).encode() + b'\x00'
     return (cmd_id, message)
 
diff --git a/Engine/Source/Programs/SwitchboardListener/SblCore/Internal/SwitchboardListener.h b/Engine/Source/Programs/SwitchboardListener/SblCore/Internal/SwitchboardListener.h
--- a/Engine/Source/Programs/SwitchboardListener/SblCore/Internal/SwitchboardListener.h
+++ b/Engine/Source/Programs/SwitchboardListener/SblCore/Internal/SwitchboardListener.h
@@ -327,10 +327,10 @@ private:
 		/** UUID of the process whose log is being streamed */
 		FGuid AssociatedProcessUUID;
 
-		FLiveLogStream(const FString& InLogFilePath, const FIPv4Endpoint& InClientEndpoint, bool bInTailOnly, const FGuid& InProcessUUID = FGuid())
+		FLiveLogStream(const FString& InLogFilePath, const FIPv4Endpoint& InClientEndpoint, bool bInTailOnly, const FGuid& InProcessUUID = FGuid(), int64 InStartOffset = -1)
 			: LogFilePath(InLogFilePath)
 			, ClientEndpoint(InClientEndpoint)
-			, LastFilePosition(0)
+			, LastFilePosition(InStartOffset >= 0 ? InStartOffset : 0)
 			, bTailOnly(bInTailOnly)
 			, LastUpdateTime(0.0)
 			, AssociatedProcessUUID(InProcessUUID)
diff --git a/Engine/Source/Programs/SwitchboardListener/SblCore/Internal/SwitchboardListenerVersion.h b/Engine/Source/Programs/SwitchboardListener/SblCore/Internal/SwitchboardListenerVersion.h
--- a/Engine/Source/Programs/SwitchboardListener/SblCore/Internal/SwitchboardListenerVersion.h
+++ b/Engine/Source/Programs/SwitchboardListener/SblCore/Internal/SwitchboardListenerVersion.h
@@ -3,5 +3,5 @@
 #pragma once
 
 #define SBLISTENER_VERSION_MAJOR 3
-#define SBLISTENER_VERSION_MINOR 4
+#define SBLISTENER_VERSION_MINOR 5
 #define SBLISTENER_VERSION_PATCH 0
diff --git a/Engine/Source/Programs/SwitchboardListener/SblCore/Private/SwitchboardListener.cpp b/Engine/Source/Programs/SwitchboardListener/SblCore/Private/SwitchboardListener.cpp
--- a/Engine/Source/Programs/SwitchboardListener/SblCore/Private/SwitchboardListener.cpp
+++ b/Engine/Source/Programs/SwitchboardListener/SblCore/Private/SwitchboardListener.cpp
@@ -2701,7 +2701,7 @@ bool FSwitchboardListener::Task_StartLiveLogStream(const FSwitchboardStartLiveLo
 		}
 
 		// Add new stream
-		FLiveLogStream NewStream(LogFilePath, ClientEndpoint, InStartLiveLogStreamTask.bTailOnly, InStartLiveLogStreamTask.ProcessUUID);
+		FLiveLogStream NewStream(LogFilePath, ClientEndpoint, InStartLiveLogStreamTask.bTailOnly, InStartLiveLogStreamTask.ProcessUUID, InStartLiveLogStreamTask.StartOffset);
 
 		if (InStartLiveLogStreamTask.ProcessUUID.IsValid())
 		{
@@ -2712,9 +2712,16 @@ bool FSwitchboardListener::Task_StartLiveLogStream(const FSwitchboardStartLiveLo
 			UE_LOG(LogSwitchboard, Log, TEXT("Starting live log stream for %s with no associated process UUID"), *LogFilePath);
 		}
 
-		// Start from the beginning unless bTailOnly is true.
-		if (!NewStream.bTailOnly)
+		// Handle start position based on StartOffset parameter
+		if (InStartLiveLogStreamTask.StartOffset >= 0)
 		{
+			// Use explicit start offset if provided
+			NewStream.LastFilePosition = InStartLiveLogStreamTask.StartOffset;
+			UE_LOG(LogSwitchboard, Log, TEXT("Starting live log stream from byte offset: %lld"), NewStream.LastFilePosition);
+		}
+		else if (!NewStream.bTailOnly)
+		{
+			// Start from the beginning
 			NewStream.LastFilePosition = 0;
 		}
 		else
diff --git a/Engine/Source/Programs/SwitchboardListener/SblCore/Private/SwitchboardProtocol.cpp b/Engine/Source/Programs/SwitchboardListener/SblCore/Private/SwitchboardProtocol.cpp
--- a/Engine/Source/Programs/SwitchboardListener/SblCore/Private/SwitchboardProtocol.cpp
+++ b/Engine/Source/Programs/SwitchboardListener/SblCore/Private/SwitchboardProtocol.cpp
@@ -470,8 +470,17 @@ FCreateTaskResult CreateTaskFromCommand(const FString& InCommand, const FIPv4End
 				FGuid::Parse(ProcessUUIDField->AsString(), ProcessUUID);
 			}
 
+			// Optional start_offset field for resuming log streaming
+
+			int64 StartOffset = -1; // -1 means use default behavior
+			TSharedPtr<FJsonValue> StartOffsetField = JsonData->TryGetField(TEXT("start_offset"));
+			if (StartOffsetField.IsValid())
+			{
+				StartOffset = (int64)StartOffsetField->AsNumber();
+			}
+
 			Result.Status = ECreateTaskStatus::Success;
-			Result.Task = MakeUnique<FSwitchboardStartLiveLogStreamTask>(MessageID, InEndpoint, LogFilePathField->AsString(), bTailOnly, ProcessUUID);
+			Result.Task = MakeUnique<FSwitchboardStartLiveLogStreamTask>(MessageID, InEndpoint, LogFilePathField->AsString(), bTailOnly, ProcessUUID, StartOffset);
 			return Result;
 		}
 	}
diff --git a/Engine/Source/Programs/SwitchboardListener/SblCore/Private/SwitchboardTasks.h b/Engine/Source/Programs/SwitchboardListener/SblCore/Private/SwitchboardTasks.h
--- a/Engine/Source/Programs/SwitchboardListener/SblCore/Private/SwitchboardTasks.h
+++ b/Engine/Source/Programs/SwitchboardListener/SblCore/Private/SwitchboardTasks.h
@@ -371,17 +371,19 @@ struct FSwitchboardFreeListenerBinaryTask : public FSwitchboardTask
 
 struct FSwitchboardStartLiveLogStreamTask : public FSwitchboardTask
 {
-	FSwitchboardStartLiveLogStreamTask(const FGuid& InTaskID, const FIPv4Endpoint& InEndpoint, const FString& InLogFilePath, bool bInTailOnly = false, const FGuid& InProcessUUID = FGuid())
+	FSwitchboardStartLiveLogStreamTask(const FGuid& InTaskID, const FIPv4Endpoint& InEndpoint, const FString& InLogFilePath, bool bInTailOnly = false, const FGuid& InProcessUUID = FGuid(), int64 InStartOffset = -1)
 		: FSwitchboardTask{ ESwitchboardTaskType::StartLiveLogStream, InTaskID, InEndpoint }
 		, LogFilePath(InLogFilePath)
 		, bTailOnly(bInTailOnly)
 		, ProcessUUID(InProcessUUID)
-		
+		, StartOffset(InStartOffset)
+
 	{}
 
 	FString LogFilePath;
 	bool bTailOnly;
 	FGuid ProcessUUID;
+	int64 StartOffset; // -1 means use default behavior (0 or file size based on bTailOnly)
 
 	//~ Begin FSwitchboardTask interface
 	static constexpr const TCHAR* CommandName = TEXT("start live log stream");
