Print out the launch arguments for the Multi-user server when it is started.

#rnx
#Synced-CL 0

diff --git a/Engine/Plugins/VirtualProduction/Switchboard/Source/Switchboard/switchboard/switchboard_application.py b/Engine/Plugins/VirtualProduction/Switchboard/Source/Switchboard/switchboard/switchboard_application.py
--- a/Engine/Plugins/VirtualProduction/Switchboard/Source/Switchboard/switchboard/switchboard_application.py
+++ b/Engine/Plugins/VirtualProduction/Switchboard/Source/Switchboard/switchboard/switchboard_application.py
@@ -220,16 +220,18 @@ class MultiUserApplication:
 
             args = list(filter(None, args))
 
+            launch_cmd_full = ' '.join(args)
             if sys.platform.startswith('win'):
                 # Given a list on Windows, Popen calls subprocess.list2cmdline
                 # to combine to a string. However, that also has the effect
                 # of turning `-foo="bar baz"`` into `"-foo=\"bar baz\""`,
                 # and Unreal's FParse won't accept that (OR `"-foo=bar baz"`);
                 # its tokenization is too different from MSVCRT.
-                args = ' '.join(args)
+                args = launch_cmd_full
+        
+            LOGGER.info(f'Launching Multi-user Server: {launch_cmd_full}');
 
             self._process = subprocess.Popen(args, creationflags=creationflags)
-
             return True
 
     def extract_process_info(self):
