Switchboard: Add recording stopwatch to status line.

This also includes some minor fixups to reduce the amount of required editing for the uic-generated switchboard_ui.py, and documents the remainder:
* Call `setupUi` from `SwitchboardDialog.__init__` rather than modifying `Ui_switchboard_window.__init__`.
* Pass the `Ui_switchboard_window` object as the argument to `setupUi`, rather than modifying it to operate on `self`.
* Change `take_spin_box` XML horizontal size policy from `Preferred` to `Fixed`, which is consistent with the previously checked-in compiled UI, and prevents a visual regression.

#Synced-CL 48625261

diff --git a/Engine/Plugins/VirtualProduction/Switchboard/Source/Switchboard/switchboard/switchboard_dialog.py b/Engine/Plugins/VirtualProduction/Switchboard/Source/Switchboard/switchboard/switchboard_dialog.py
--- a/Engine/Plugins/VirtualProduction/Switchboard/Source/Switchboard/switchboard/switchboard_dialog.py
+++ b/Engine/Plugins/VirtualProduction/Switchboard/Source/Switchboard/switchboard/switchboard_dialog.py
@@ -280,7 +280,8 @@ class SwitchboardDialog(QtCore.QObject):
         ConsoleStream.stderr().message_written.connect(self._console_pipe)
 
         # Create an instance of the UI class
-        self.window = Ui_switchboard_window()
+        self.window: Ui_switchboard_window = Ui_switchboard_window()
+        self.window.setupUi(self.window)
 
         # Accept copy-paste, intended for uassets or umaps from Unreal Engine's Content Browser
         QtGui.QShortcut(QtGui.QKeySequence.Paste, self.window).activated.connect(self.on_paste)
@@ -310,6 +311,11 @@ class SwitchboardDialog(QtCore.QObject):
         self._level = None
         self._multiuser_session_name = None
         self._is_recording = False
+        self._recording_start_time: datetime.datetime | None = None
+        self._recording_end_time: datetime.datetime | None = None
+        self._recording_timer = QtCore.QTimer(self)
+        self._recording_timer.timeout.connect(self._on_recording_timer)
+        self._update_recording_status()
         self._description = 'description'
 
         # Recording Manager
@@ -2551,6 +2557,33 @@ class SwitchboardDialog(QtCore.QObject):
 
         self._record_start(self.slate, self.take, self.description, exclude_address=device.address)
 
+    def _on_recording_timer(self):
+        self._update_recording_status()
+
+    def _fmt_rec_duration(self, timedelta: datetime.timedelta) -> str:
+        mm, ss = divmod(int(timedelta.total_seconds() + 0.5), 60)
+        hh, mm = divmod(mm, 60)
+        result = ''
+        if hh:
+            result += f'{hh:02}' + (' hours, ' if hh != 1 else ' hour, ')
+        if mm or hh:
+            result += f'{mm:02}' + (' mins, ' if mm != 1 else ' min, ')
+        result += f'{ss:02}'  + (' secs ' if ss != 1 else ' sec ')
+        return result
+
+    def _update_recording_status(self):
+        self.window.recording_status.setStyleSheet('')
+
+        if not self._recording_start_time:
+            self.window.recording_status.setText('Not Recording ')
+        elif self._recording_end_time:
+            recording_duration = self._recording_end_time - self._recording_start_time
+            self.window.recording_status.setText('Last Recording: ' + self._fmt_rec_duration(recording_duration))
+        else:
+            recording_duration = datetime.datetime.now() - self._recording_start_time
+            self.window.recording_status.setText('Recording Duration: ' + self._fmt_rec_duration(recording_duration))
+            self.window.recording_status.setStyleSheet('color: #e52d00;')
+
     def _record_start(self, slate, take, description, exclude_address=None):
         LOGGER.success(f'Record Start: "{self.slate}" {self.take}')
 
@@ -2563,6 +2596,10 @@ class SwitchboardDialog(QtCore.QObject):
 
         # Start a new recording
         self._is_recording = True
+        self._recording_start_time = datetime.datetime.now()
+        self._recording_end_time = None
+        self._recording_timer.start(1000)
+        self._update_recording_status()
 
         # TODO: Lock SLATE/TAKE/SESSION/CL
 
@@ -2574,7 +2611,7 @@ class SwitchboardDialog(QtCore.QObject):
         new_recording.slate = self.slate
         new_recording.take = self.take
         new_recording.description = self.description
-        new_recording.date = switchboard_utils.date_to_string(datetime.date.today())
+        new_recording.date = switchboard_utils.date_to_string(self._recording_start_time.date())
         new_recording.map = self.level
         new_recording.multiuser_session = self.multiuser_session_name()
         new_recording.changelist = self.project_changelist
@@ -2606,6 +2643,9 @@ class SwitchboardDialog(QtCore.QObject):
 
         # End the recording
         self._is_recording = False
+        self._recording_end_time = datetime.datetime.now()
+        self._recording_timer.stop()
+        self._update_recording_status()
 
         # Pull the latest recording down
         new_recording = self.recording_manager.latest_recording
diff --git a/Engine/Plugins/VirtualProduction/Switchboard/Source/Switchboard/switchboard/ui/switchboard_ui.py b/Engine/Plugins/VirtualProduction/Switchboard/Source/Switchboard/switchboard/ui/switchboard_ui.py
--- a/Engine/Plugins/VirtualProduction/Switchboard/Source/Switchboard/switchboard/ui/switchboard_ui.py
+++ b/Engine/Plugins/VirtualProduction/Switchboard/Source/Switchboard/switchboard/ui/switchboard_ui.py
@@ -1,27 +1,37 @@
 # Copyright Epic Games, Inc. All Rights Reserved.
+#
+# To regenerate:
+#     > cd (...)\Engine\Plugins\VirtualProduction\Switchboard\Source\Switchboard\switchboard
+#     > ..\..\..\..\..\..\Extras\ThirdPartyNotUE\SwitchboardThirdParty\Python\Scripts\pyside6-uic.exe ui\switchboard.ui -o ui\switchboard_ui.py
+#
+# Manual fix-ups:
+#  * Add this comment block, including copyright, to the top of the file
+#  * Fix generated `from switchboard` import paths
+#  * Remove `import resources_rc`
+#  * Change Ui_switchboard_window base class from `object` to `QMainWindow`
 
 ################################################################################
 ## Form generated from reading UI file 'switchboard.ui'
 ##
-## Created by: Qt User Interface Compiler version 6.8.0
-## Then fixed manually.
+## Created by: Qt User Interface Compiler version 6.5.3
 ##
 ## WARNING! All changes made in this file will be lost when recompiling UI file!
 ################################################################################
 
-from PySide6.QtCore import (
-    QCoreApplication, QMetaObject, QRect, QSize, Qt
-)
-
-from PySide6.QtGui import (
-    QAction, QBrush, QColor, QFont, QGradient, QIcon, QLinearGradient, QPalette
-)
-
-from PySide6.QtWidgets import (
-    QCheckBox, QComboBox, QFrame, QHBoxLayout, QLabel, QLineEdit, QMainWindow,
-    QMenu, QMenuBar, QPlainTextEdit, QPushButton, QSizePolicy, QSpacerItem,
-    QSpinBox, QSplitter, QStatusBar, QTabWidget, QToolButton, QVBoxLayout, QWidget
-)
+from PySide6.QtCore import (QCoreApplication, QDate, QDateTime, QLocale,
+    QMetaObject, QObject, QPoint, QRect,
+    QSize, QTime, QUrl, Qt)
+from PySide6.QtGui import (QAction, QBrush, QColor, QConicalGradient,
+    QCursor, QFont, QFontDatabase, QGradient,
+    QIcon, QImage, QKeySequence, QLinearGradient,
+    QPainter, QPalette, QPixmap, QRadialGradient,
+    QTransform)
+from PySide6.QtWidgets import (QApplication, QCheckBox, QComboBox, QFrame,
+    QHBoxLayout, QLabel, QLineEdit, QListWidgetItem,
+    QMainWindow, QMenu, QMenuBar, QPlainTextEdit,
+    QPushButton, QSizePolicy, QSpacerItem, QSplitter,
+    QStatusBar, QTabWidget, QToolButton, QVBoxLayout,
+    QWidget)
 
 from switchboard.device_list_widget import DeviceListWidget
 from switchboard.switchboard_widgets import (ControlQPushButton, FramelessQLineEdit)
@@ -29,38 +39,33 @@ from switchboard.ui.take_spin_box import TakeSpinBox
 
 
 class Ui_switchboard_window(QMainWindow):
-
-    def __init__(self):
-        super(Ui_switchboard_window, self).__init__()
-        self.setupUi()
-
-    def setupUi(self):
-        if not self.objectName():
-            self.setObjectName(u"switchboard_window")
-        self.resize(1280, 952)
-        self.setMinimumSize(QSize(900, 500))
+    def setupUi(self, switchboard_window):
+        if not switchboard_window.objectName():
+            switchboard_window.setObjectName(u"switchboard_window")
+        switchboard_window.resize(1280, 952)
+        switchboard_window.setMinimumSize(QSize(900, 500))
         icon = QIcon()
-        icon.addFile(u":/icons/images/switchboard.png", QSize(), QIcon.Mode.Normal, QIcon.State.Off)
-        self.setWindowIcon(icon)
-        self.menu_p4_refresh_cl = QAction(self)
+        icon.addFile(u":/icons/images/switchboard.png", QSize(), QIcon.Normal, QIcon.Off)
+        switchboard_window.setWindowIcon(icon)
+        self.menu_p4_refresh_cl = QAction(switchboard_window)
         self.menu_p4_refresh_cl.setObjectName(u"menu_p4_refresh_cl")
-        self.actionMinimal_Sync = QAction(self)
+        self.actionMinimal_Sync = QAction(switchboard_window)
         self.actionMinimal_Sync.setObjectName(u"actionMinimal_Sync")
-        self.update_settings = QAction(self)
+        self.update_settings = QAction(switchboard_window)
         self.update_settings.setObjectName(u"update_settings")
-        self.actionAsd = QAction(self)
+        self.actionAsd = QAction(switchboard_window)
         self.actionAsd.setObjectName(u"actionAsd")
-        self.menu_new_config = QAction(self)
+        self.menu_new_config = QAction(switchboard_window)
         self.menu_new_config.setObjectName(u"menu_new_config")
-        self.menu_save_config_as = QAction(self)
+        self.menu_save_config_as = QAction(switchboard_window)
         self.menu_save_config_as.setObjectName(u"menu_save_config_as")
-        self.menu_delete_config = QAction(self)
+        self.menu_delete_config = QAction(switchboard_window)
         self.menu_delete_config.setObjectName(u"menu_delete_config")
-        self.actionTest = QAction(self)
+        self.actionTest = QAction(switchboard_window)
         self.actionTest.setObjectName(u"actionTest")
-        self.centralwidget = QWidget(self)
+        self.centralwidget = QWidget(switchboard_window)
         self.centralwidget.setObjectName(u"centralwidget")
-        sizePolicy = QSizePolicy(QSizePolicy.Policy.Preferred, QSizePolicy.Policy.Preferred)
+        sizePolicy = QSizePolicy(QSizePolicy.Preferred, QSizePolicy.Preferred)
         sizePolicy.setHorizontalStretch(0)
         sizePolicy.setVerticalStretch(0)
         sizePolicy.setHeightForWidth(self.centralwidget.sizePolicy().hasHeightForWidth())
@@ -79,13 +84,13 @@ class Ui_switchboard_window(QMainWindow):
         font.setPointSize(9)
         font.setBold(True)
         self.label_4.setFont(font)
-        self.label_4.setAlignment(Qt.AlignRight | Qt.AlignTrailing | Qt.AlignVCenter)
+        self.label_4.setAlignment(Qt.AlignRight|Qt.AlignTrailing|Qt.AlignVCenter)
 
         self.horizontalLayout.addWidget(self.label_4)
 
         self.sequence_line_edit = FramelessQLineEdit(self.centralwidget)
         self.sequence_line_edit.setObjectName(u"sequence_line_edit")
-        sizePolicy1 = QSizePolicy(QSizePolicy.Policy.MinimumExpanding, QSizePolicy.Policy.Fixed)
+        sizePolicy1 = QSizePolicy(QSizePolicy.MinimumExpanding, QSizePolicy.Fixed)
         sizePolicy1.setHorizontalStretch(0)
         sizePolicy1.setVerticalStretch(0)
         sizePolicy1.setHeightForWidth(self.sequence_line_edit.sizePolicy().hasHeightForWidth())
@@ -95,14 +100,14 @@ class Ui_switchboard_window(QMainWindow):
         font1 = QFont()
         font1.setPointSize(20)
         self.sequence_line_edit.setFont(font1)
-        self.sequence_line_edit.setProperty(u"recording", True)
+        self.sequence_line_edit.setProperty("recording", True)
 
         self.horizontalLayout.addWidget(self.sequence_line_edit)
 
         self.slate_label = QLabel(self.centralwidget)
         self.slate_label.setObjectName(u"slate_label")
         self.slate_label.setEnabled(True)
-        sizePolicy2 = QSizePolicy(QSizePolicy.Policy.Fixed, QSizePolicy.Policy.Preferred)
+        sizePolicy2 = QSizePolicy(QSizePolicy.Fixed, QSizePolicy.Preferred)
         sizePolicy2.setHorizontalStretch(0)
         sizePolicy2.setVerticalStretch(0)
         sizePolicy2.setHeightForWidth(self.slate_label.sizePolicy().hasHeightForWidth())
@@ -123,7 +128,7 @@ class Ui_switchboard_window(QMainWindow):
         font2.setFamilies([u"Roboto"])
         font2.setPointSize(20)
         self.slate_line_edit.setFont(font2)
-        self.slate_line_edit.setAlignment(Qt.AlignLeading | Qt.AlignLeft | Qt.AlignVCenter)
+        self.slate_line_edit.setAlignment(Qt.AlignLeading|Qt.AlignLeft|Qt.AlignVCenter)
 
         self.horizontalLayout.addWidget(self.slate_line_edit)
 
@@ -140,13 +145,13 @@ class Ui_switchboard_window(QMainWindow):
 
         self.take_spin_box = TakeSpinBox(self.centralwidget)
         self.take_spin_box.setObjectName(u"take_spin_box")
-        sizePolicy3 = QSizePolicy(QSizePolicy.Policy.Fixed, QSizePolicy.Policy.Fixed)
+        sizePolicy3 = QSizePolicy(QSizePolicy.Fixed, QSizePolicy.Fixed)
         sizePolicy3.setHorizontalStretch(0)
         sizePolicy3.setVerticalStretch(0)
         sizePolicy3.setHeightForWidth(self.take_spin_box.sizePolicy().hasHeightForWidth())
         self.take_spin_box.setSizePolicy(sizePolicy3)
         self.take_spin_box.setMinimumSize(QSize(80, 0))
-        self.take_spin_box.setMaximumSize(QSize(80, 16777215))
+        self.take_spin_box.setMaximumSize(QSize(16777215, 16777215))
         self.take_spin_box.setFont(font2)
         self.take_spin_box.setAlignment(Qt.AlignCenter)
         self.take_spin_box.setMinimum(1)
@@ -157,30 +162,29 @@ class Ui_switchboard_window(QMainWindow):
         self.record_button = QPushButton(self.centralwidget)
         self.record_button.setObjectName(u"record_button")
         self.record_button.setEnabled(True)
-        sizePolicy4 = QSizePolicy(QSizePolicy.Policy.Fixed, QSizePolicy.Policy.Fixed)
-        sizePolicy4.setHorizontalStretch(0)
-        sizePolicy4.setVerticalStretch(0)
-        sizePolicy4.setHeightForWidth(self.record_button.sizePolicy().hasHeightForWidth())
-        self.record_button.setSizePolicy(sizePolicy4)
+        sizePolicy3.setHeightForWidth(self.record_button.sizePolicy().hasHeightForWidth())
+        self.record_button.setSizePolicy(sizePolicy3)
         self.record_button.setMinimumSize(QSize(75, 75))
         self.record_button.setMaximumSize(QSize(75, 75))
         self.record_button.setBaseSize(QSize(0, 0))
         icon1 = QIcon()
-        icon1.addFile(u":/icons/images/record_stop.png", QSize(), QIcon.Mode.Normal, QIcon.State.On)
+        icon1.addFile(u":/icons/images/record_stop.png", QSize(), QIcon.Normal, QIcon.On)
         self.record_button.setIcon(icon1)
         self.record_button.setIconSize(QSize(60, 60))
         self.record_button.setCheckable(True)
         self.record_button.setChecked(False)
-        self.record_button.setProperty(u"no_background", True)
-        self.record_button.setProperty(u"no_border", True)
+        self.record_button.setProperty("no_background", True)
+        self.record_button.setProperty("no_border", True)
 
         self.horizontalLayout.addWidget(self.record_button)
+
+
         self.verticalLayout_3.addLayout(self.horizontalLayout)
 
         self.line = QFrame(self.centralwidget)
         self.line.setObjectName(u"line")
-        self.line.setFrameShape(QFrame.Shape.HLine)
-        self.line.setFrameShadow(QFrame.Shadow.Sunken)
+        self.line.setFrameShape(QFrame.HLine)
+        self.line.setFrameShadow(QFrame.Sunken)
 
         self.verticalLayout_3.addWidget(self.line)
 
@@ -200,7 +204,7 @@ class Ui_switchboard_window(QMainWindow):
         self.settings_button.setMinimumSize(QSize(0, 0))
         self.settings_button.setStyleSheet(u"QPushButton { background-color: transparent; border: 0px }")
         icon2 = QIcon()
-        icon2.addFile(u":/icons/images/Settings_40x.png", QSize(), QIcon.Mode.Normal, QIcon.State.On)
+        icon2.addFile(u":/icons/images/Settings_40x.png", QSize(), QIcon.Normal, QIcon.On)
         self.settings_button.setIcon(icon2)
 
         self.horizontalLayout_7.addWidget(self.settings_button)
@@ -213,7 +217,7 @@ class Ui_switchboard_window(QMainWindow):
 
         self.horizontalLayout_7.addWidget(self.device_add_tool_button)
 
-        self.horizontalSpacer_2 = QSpacerItem(50, 20, QSizePolicy.Policy.Minimum, QSizePolicy.Policy.Minimum)
+        self.horizontalSpacer_2 = QSpacerItem(50, 20, QSizePolicy.Minimum, QSizePolicy.Minimum)
 
         self.horizontalLayout_7.addItem(self.horizontalSpacer_2)
 
@@ -224,10 +228,11 @@ class Ui_switchboard_window(QMainWindow):
         font3.setFamilies([u"Roboto"])
         font3.setBold(True)
         self.label_2.setFont(font3)
-        self.label_2.setAlignment(Qt.AlignRight | Qt.AlignTrailing | Qt.AlignVCenter)
+        self.label_2.setAlignment(Qt.AlignRight|Qt.AlignTrailing|Qt.AlignVCenter)
 
         self.horizontalLayout_7.addWidget(self.label_2)
 
+
         self.verticalLayout.addLayout(self.horizontalLayout_7)
 
         self.horizontalLayout_3 = QHBoxLayout()
@@ -245,7 +250,7 @@ class Ui_switchboard_window(QMainWindow):
         font4.setFamilies([u"Roboto"])
         font4.setBold(False)
         self.engine_cl_label.setFont(font4)
-        self.engine_cl_label.setAlignment(Qt.AlignLeading | Qt.AlignLeft | Qt.AlignVCenter)
+        self.engine_cl_label.setAlignment(Qt.AlignLeading|Qt.AlignLeft|Qt.AlignVCenter)
 
         self.horizontalLayout_10.addWidget(self.engine_cl_label)
 
@@ -261,15 +266,16 @@ class Ui_switchboard_window(QMainWindow):
         self.refresh_engine_cl_button.setObjectName(u"refresh_engine_cl_button")
         self.refresh_engine_cl_button.setMaximumSize(QSize(22, 22))
         icon3 = QIcon()
-        icon3.addFile(u":/icons/images/icon_refresh.png", QSize(), QIcon.Mode.Normal, QIcon.State.Off)
-        icon3.addFile(u":/icons/images/icon_refresh_disabled.png", QSize(), QIcon.Mode.Disabled, QIcon.State.Off)
-        icon3.addFile(u":/icons/images/icon_refresh_hover.png", QSize(), QIcon.Mode.Active, QIcon.State.Off)
+        icon3.addFile(u":/icons/images/icon_refresh.png", QSize(), QIcon.Normal, QIcon.Off)
+        icon3.addFile(u":/icons/images/icon_refresh_disabled.png", QSize(), QIcon.Disabled, QIcon.Off)
+        icon3.addFile(u":/icons/images/icon_refresh_hover.png", QSize(), QIcon.Active, QIcon.Off)
         self.refresh_engine_cl_button.setIcon(icon3)
         self.refresh_engine_cl_button.setIconSize(QSize(18, 18))
-        self.refresh_engine_cl_button.setProperty(u"frameless", True)
+        self.refresh_engine_cl_button.setProperty("frameless", True)
 
         self.horizontalLayout_10.addWidget(self.refresh_engine_cl_button)
 
+
         self.verticalLayout_6.addLayout(self.horizontalLayout_10)
 
         self.horizontalLayout_9 = QHBoxLayout()
@@ -278,15 +284,18 @@ class Ui_switchboard_window(QMainWindow):
         self.project_cl_label.setObjectName(u"project_cl_label")
         self.project_cl_label.setMinimumSize(QSize(60, 0))
         self.project_cl_label.setFont(font4)
-        self.project_cl_label.setAlignment(Qt.AlignLeading | Qt.AlignLeft | Qt.AlignVCenter)
+        self.project_cl_label.setAlignment(Qt.AlignLeading|Qt.AlignLeft|Qt.AlignVCenter)
 
         self.horizontalLayout_9.addWidget(self.project_cl_label)
 
         self.project_cl_combo_box = QComboBox(self.device_widget)
         self.project_cl_combo_box.addItem("")
         self.project_cl_combo_box.setObjectName(u"project_cl_combo_box")
-        sizePolicy3.setHeightForWidth(self.project_cl_combo_box.sizePolicy().hasHeightForWidth())
-        self.project_cl_combo_box.setSizePolicy(sizePolicy3)
+        sizePolicy4 = QSizePolicy(QSizePolicy.Preferred, QSizePolicy.Fixed)
+        sizePolicy4.setHorizontalStretch(0)
+        sizePolicy4.setVerticalStretch(0)
+        sizePolicy4.setHeightForWidth(self.project_cl_combo_box.sizePolicy().hasHeightForWidth())
+        self.project_cl_combo_box.setSizePolicy(sizePolicy4)
         self.project_cl_combo_box.setMinimumSize(QSize(80, 0))
         self.project_cl_combo_box.setMaximumSize(QSize(16777215, 16777215))
         font5 = QFont()
@@ -300,122 +309,128 @@ class Ui_switchboard_window(QMainWindow):
 
         self.refresh_project_cl_button = ControlQPushButton(self.device_widget)
         self.refresh_project_cl_button.setObjectName(u"refresh_project_cl_button")
-        sizePolicy3.setHeightForWidth(self.refresh_project_cl_button.sizePolicy().hasHeightForWidth())
-        self.refresh_project_cl_button.setSizePolicy(sizePolicy3)
+        sizePolicy4.setHeightForWidth(self.refresh_project_cl_button.sizePolicy().hasHeightForWidth())
+        self.refresh_project_cl_button.setSizePolicy(sizePolicy4)
         self.refresh_project_cl_button.setMaximumSize(QSize(22, 22))
         self.refresh_project_cl_button.setIcon(icon3)
         self.refresh_project_cl_button.setIconSize(QSize(18, 18))
-        self.refresh_project_cl_button.setProperty(u"frameless", True)
+        self.refresh_project_cl_button.setProperty("frameless", True)
 
         self.horizontalLayout_9.addWidget(self.refresh_project_cl_button)
+
+
         self.verticalLayout_6.addLayout(self.horizontalLayout_9)
+
+
         self.horizontalLayout_3.addLayout(self.verticalLayout_6)
 
         self.sync_all_button = ControlQPushButton(self.device_widget)
         self.sync_all_button.setObjectName(u"sync_all_button")
-        sizePolicy3.setHeightForWidth(self.sync_all_button.sizePolicy().hasHeightForWidth())
-        self.sync_all_button.setSizePolicy(sizePolicy3)
+        sizePolicy4.setHeightForWidth(self.sync_all_button.sizePolicy().hasHeightForWidth())
+        self.sync_all_button.setSizePolicy(sizePolicy4)
         self.sync_all_button.setMinimumSize(QSize(25, 30))
         icon4 = QIcon()
-        icon4.addFile(u":/icons/images/icon_sync.png", QSize(), QIcon.Mode.Normal, QIcon.State.On)
-        icon4.addFile(u":/icons/images/icon_sync_disabled.png", QSize(), QIcon.Mode.Disabled, QIcon.State.On)
-        icon4.addFile(u":/icons/images/icon_sync_hover.png", QSize(), QIcon.Mode.Active, QIcon.State.On)
+        icon4.addFile(u":/icons/images/icon_sync.png", QSize(), QIcon.Normal, QIcon.On)
+        icon4.addFile(u":/icons/images/icon_sync_disabled.png", QSize(), QIcon.Disabled, QIcon.On)
+        icon4.addFile(u":/icons/images/icon_sync_hover.png", QSize(), QIcon.Active, QIcon.On)
         self.sync_all_button.setIcon(icon4)
         self.sync_all_button.setIconSize(QSize(21, 21))
-        self.sync_all_button.setProperty(u"frameless", True)
+        self.sync_all_button.setProperty("frameless", True)
 
         self.horizontalLayout_3.addWidget(self.sync_all_button)
 
         self.build_all_button = ControlQPushButton(self.device_widget)
         self.build_all_button.setObjectName(u"build_all_button")
-        sizePolicy3.setHeightForWidth(self.build_all_button.sizePolicy().hasHeightForWidth())
-        self.build_all_button.setSizePolicy(sizePolicy3)
+        sizePolicy4.setHeightForWidth(self.build_all_button.sizePolicy().hasHeightForWidth())
+        self.build_all_button.setSizePolicy(sizePolicy4)
         self.build_all_button.setMinimumSize(QSize(25, 30))
         icon5 = QIcon()
-        icon5.addFile(u":/icons/images/icon_build.png", QSize(), QIcon.Mode.Normal, QIcon.State.On)
-        icon5.addFile(u":/icons/images/icon_build_hover.png", QSize(), QIcon.Mode.Active, QIcon.State.On)
+        icon5.addFile(u":/icons/images/icon_build.png", QSize(), QIcon.Normal, QIcon.On)
+        icon5.addFile(u":/icons/images/icon_build_hover.png", QSize(), QIcon.Active, QIcon.On)
         self.build_all_button.setIcon(icon5)
         self.build_all_button.setIconSize(QSize(21, 21))
-        self.build_all_button.setProperty(u"frameless", True)
+        self.build_all_button.setProperty("frameless", True)
 
         self.horizontalLayout_3.addWidget(self.build_all_button)
 
         self.sync_and_build_all_button = ControlQPushButton(self.device_widget)
         self.sync_and_build_all_button.setObjectName(u"sync_and_build_all_button")
-        sizePolicy3.setHeightForWidth(self.sync_and_build_all_button.sizePolicy().hasHeightForWidth())
-        self.sync_and_build_all_button.setSizePolicy(sizePolicy3)
+        sizePolicy4.setHeightForWidth(self.sync_and_build_all_button.sizePolicy().hasHeightForWidth())
+        self.sync_and_build_all_button.setSizePolicy(sizePolicy4)
         self.sync_and_build_all_button.setMinimumSize(QSize(25, 30))
         icon6 = QIcon()
-        icon6.addFile(u":/icons/images/icon_sync_and_build.png", QSize(), QIcon.Mode.Normal, QIcon.State.On)
-        icon6.addFile(u":/icons/images/icon_sync_and_build_disabled.png", QSize(), QIcon.Mode.Disabled, QIcon.State.On)
-        icon6.addFile(u":/icons/images/icon_sync_and_build_hover.png", QSize(), QIcon.Mode.Active, QIcon.State.On)
+        icon6.addFile(u":/icons/images/icon_sync_and_build.png", QSize(), QIcon.Normal, QIcon.On)
+        icon6.addFile(u":/icons/images/icon_sync_and_build_disabled.png", QSize(), QIcon.Disabled, QIcon.On)
+        icon6.addFile(u":/icons/images/icon_sync_and_build_hover.png", QSize(), QIcon.Active, QIcon.On)
         self.sync_and_build_all_button.setIcon(icon6)
         self.sync_and_build_all_button.setIconSize(QSize(21, 21))
-        self.sync_and_build_all_button.setProperty(u"frameless", True)
+        self.sync_and_build_all_button.setProperty("frameless", True)
 
         self.horizontalLayout_3.addWidget(self.sync_and_build_all_button)
 
-        self.horizontalSpacer = QSpacerItem(10, 20, QSizePolicy.Policy.Expanding, QSizePolicy.Policy.Minimum)
+        self.horizontalSpacer = QSpacerItem(10, 20, QSizePolicy.Expanding, QSizePolicy.Minimum)
 
         self.horizontalLayout_3.addItem(self.horizontalSpacer)
 
         self.device_settings_button = ControlQPushButton(self.device_widget)
         self.device_settings_button.setObjectName(u"device_settings_button")
         self.device_settings_button.setEnabled(False)
-        sizePolicy3.setHeightForWidth(self.device_settings_button.sizePolicy().hasHeightForWidth())
-        self.device_settings_button.setSizePolicy(sizePolicy3)
+        sizePolicy4.setHeightForWidth(self.device_settings_button.sizePolicy().hasHeightForWidth())
+        self.device_settings_button.setSizePolicy(sizePolicy4)
         self.device_settings_button.setMinimumSize(QSize(25, 30))
         icon7 = QIcon()
-        icon7.addFile(u":/icons/images/icon_settings_disabled.png", QSize(), QIcon.Mode.Normal, QIcon.State.Off)
-        icon7.addFile(u":/icons/images/icon_settings.png", QSize(), QIcon.Mode.Normal, QIcon.State.On)
-        icon7.addFile(u":/icons/images/icon_settings_hover.png", QSize(), QIcon.Mode.Active, QIcon.State.On)
+        icon7.addFile(u":/icons/images/icon_settings_disabled.png", QSize(), QIcon.Normal, QIcon.Off)
+        icon7.addFile(u":/icons/images/icon_settings.png", QSize(), QIcon.Normal, QIcon.On)
+        icon7.addFile(u":/icons/images/icon_settings_hover.png", QSize(), QIcon.Active, QIcon.On)
         self.device_settings_button.setIcon(icon7)
         self.device_settings_button.setIconSize(QSize(21, 21))
         self.device_settings_button.setCheckable(True)
-        self.device_settings_button.setProperty(u"frameless", True)
+        self.device_settings_button.setProperty("frameless", True)
 
         self.horizontalLayout_3.addWidget(self.device_settings_button)
 
         self.launch_all_button = ControlQPushButton(self.device_widget)
         self.launch_all_button.setObjectName(u"launch_all_button")
         self.launch_all_button.setEnabled(False)
-        sizePolicy3.setHeightForWidth(self.launch_all_button.sizePolicy().hasHeightForWidth())
-        self.launch_all_button.setSizePolicy(sizePolicy3)
+        sizePolicy4.setHeightForWidth(self.launch_all_button.sizePolicy().hasHeightForWidth())
+        self.launch_all_button.setSizePolicy(sizePolicy4)
         self.launch_all_button.setMinimumSize(QSize(25, 30))
         icon8 = QIcon()
-        icon8.addFile(u":/icons/images/icon_open.png", QSize(), QIcon.Mode.Normal, QIcon.State.Off)
-        icon8.addFile(u":/icons/images/icon_close.png", QSize(), QIcon.Mode.Normal, QIcon.State.On)
-        icon8.addFile(u":/icons/images/icon_open_disabled.png", QSize(), QIcon.Mode.Disabled, QIcon.State.Off)
-        icon8.addFile(u":/icons/images/icon_open_hover.png", QSize(), QIcon.Mode.Active, QIcon.State.Off)
-        icon8.addFile(u":/icons/images/icon_close_hover.png", QSize(), QIcon.Mode.Active, QIcon.State.On)
+        icon8.addFile(u":/icons/images/icon_open.png", QSize(), QIcon.Normal, QIcon.Off)
+        icon8.addFile(u":/icons/images/icon_close.png", QSize(), QIcon.Normal, QIcon.On)
+        icon8.addFile(u":/icons/images/icon_open_disabled.png", QSize(), QIcon.Disabled, QIcon.Off)
+        icon8.addFile(u":/icons/images/icon_open_hover.png", QSize(), QIcon.Active, QIcon.Off)
+        icon8.addFile(u":/icons/images/icon_close_hover.png", QSize(), QIcon.Active, QIcon.On)
         self.launch_all_button.setIcon(icon8)
         self.launch_all_button.setIconSize(QSize(21, 21))
         self.launch_all_button.setCheckable(True)
-        self.launch_all_button.setProperty(u"frameless", True)
+        self.launch_all_button.setProperty("frameless", True)
 
         self.horizontalLayout_3.addWidget(self.launch_all_button)
 
         self.connect_all_button = ControlQPushButton(self.device_widget)
         self.connect_all_button.setObjectName(u"connect_all_button")
-        sizePolicy3.setHeightForWidth(self.connect_all_button.sizePolicy().hasHeightForWidth())
-        self.connect_all_button.setSizePolicy(sizePolicy3)
+        sizePolicy4.setHeightForWidth(self.connect_all_button.sizePolicy().hasHeightForWidth())
+        self.connect_all_button.setSizePolicy(sizePolicy4)
         self.connect_all_button.setMinimumSize(QSize(25, 30))
         icon9 = QIcon()
-        icon9.addFile(u":/icons/images/icon_connect.png", QSize(), QIcon.Mode.Normal, QIcon.State.Off)
-        icon9.addFile(u":/icons/images/icon_connected.png", QSize(), QIcon.Mode.Normal, QIcon.State.On)
-        icon9.addFile(u":/icons/images/icon_connect_hover.png", QSize(), QIcon.Mode.Active, QIcon.State.Off)
-        icon9.addFile(u":/icons/images/icon_connected_hover.png", QSize(), QIcon.Mode.Active, QIcon.State.On)
+        icon9.addFile(u":/icons/images/icon_connect.png", QSize(), QIcon.Normal, QIcon.Off)
+        icon9.addFile(u":/icons/images/icon_connected.png", QSize(), QIcon.Normal, QIcon.On)
+        icon9.addFile(u":/icons/images/icon_connect_hover.png", QSize(), QIcon.Active, QIcon.Off)
+        icon9.addFile(u":/icons/images/icon_connected_hover.png", QSize(), QIcon.Active, QIcon.On)
         self.connect_all_button.setIcon(icon9)
         self.connect_all_button.setIconSize(QSize(21, 21))
         self.connect_all_button.setCheckable(True)
-        self.connect_all_button.setProperty(u"frameless", True)
+        self.connect_all_button.setProperty("frameless", True)
 
         self.horizontalLayout_3.addWidget(self.connect_all_button)
+
+
         self.verticalLayout.addLayout(self.horizontalLayout_3)
 
         self.device_list_widget = DeviceListWidget(self.device_widget)
         self.device_list_widget.setObjectName(u"device_list_widget")
-        sizePolicy5 = QSizePolicy(QSizePolicy.Policy.Expanding, QSizePolicy.Policy.Expanding)
+        sizePolicy5 = QSizePolicy(QSizePolicy.Expanding, QSizePolicy.Expanding)
         sizePolicy5.setHorizontalStretch(0)
         sizePolicy5.setVerticalStretch(0)
         sizePolicy5.setHeightForWidth(self.device_list_widget.sizePolicy().hasHeightForWidth())
@@ -428,8 +443,8 @@ class Ui_switchboard_window(QMainWindow):
 
         self.line_3 = QFrame(self.device_widget)
         self.line_3.setObjectName(u"line_3")
-        self.line_3.setFrameShape(QFrame.Shape.HLine)
-        self.line_3.setFrameShadow(QFrame.Shadow.Sunken)
+        self.line_3.setFrameShape(QFrame.HLine)
+        self.line_3.setFrameShadow(QFrame.Sunken)
 
         self.verticalLayout.addWidget(self.line_3)
 
@@ -449,7 +464,7 @@ class Ui_switchboard_window(QMainWindow):
         self.muserver_start_stop_button.setEnabled(True)
         self.muserver_start_stop_button.setIcon(icon8)
         self.muserver_start_stop_button.setCheckable(True)
-        self.muserver_start_stop_button.setProperty(u"frameless", True)
+        self.muserver_start_stop_button.setProperty("frameless", True)
 
         self.horizontalLayout_5.addWidget(self.muserver_start_stop_button)
 
@@ -462,8 +477,8 @@ class Ui_switchboard_window(QMainWindow):
 
         self.line_4 = QFrame(self.device_widget)
         self.line_4.setObjectName(u"line_4")
-        self.line_4.setFrameShape(QFrame.Shape.VLine)
-        self.line_4.setFrameShadow(QFrame.Shadow.Sunken)
+        self.line_4.setFrameShape(QFrame.VLine)
+        self.line_4.setFrameShadow(QFrame.Sunken)
 
         self.horizontalLayout_5.addWidget(self.line_4)
 
@@ -484,26 +499,30 @@ class Ui_switchboard_window(QMainWindow):
         font7 = QFont()
         font7.setPointSize(8)
         self.multiuser_session_lineEdit.setFont(font7)
-        self.multiuser_session_lineEdit.setProperty(u"recording", True)
+        self.multiuser_session_lineEdit.setProperty("recording", True)
 
         self.horizontalLayout_5.addWidget(self.multiuser_session_lineEdit)
 
         self.multiuser_session_inc_button = QPushButton(self.device_widget)
         self.multiuser_session_inc_button.setObjectName(u"multiuser_session_inc_button")
         self.multiuser_session_inc_button.setEnabled(True)
-        sizePolicy6 = QSizePolicy(QSizePolicy.Policy.Minimum, QSizePolicy.Policy.Preferred)
+        sizePolicy6 = QSizePolicy(QSizePolicy.Minimum, QSizePolicy.Preferred)
         sizePolicy6.setHorizontalStretch(0)
         sizePolicy6.setVerticalStretch(0)
         sizePolicy6.setHeightForWidth(self.multiuser_session_inc_button.sizePolicy().hasHeightForWidth())
         self.multiuser_session_inc_button.setSizePolicy(sizePolicy6)
         icon10 = QIcon()
-        icon10.addFile(u":/icons/images/icon_add_disabled.png", QSize(), QIcon.Mode.Normal, QIcon.State.Off)
-        icon10.addFile(u":/icons/images/icon_add.png", QSize(), QIcon.Mode.Normal, QIcon.State.On)
+        icon10.addFile(u":/icons/images/icon_add_disabled.png", QSize(), QIcon.Normal, QIcon.Off)
+        icon10.addFile(u":/icons/images/icon_add.png", QSize(), QIcon.Normal, QIcon.On)
         self.multiuser_session_inc_button.setIcon(icon10)
-        self.multiuser_session_inc_button.setProperty(u"frameless", True)
+        self.multiuser_session_inc_button.setProperty("frameless", True)
 
         self.horizontalLayout_5.addWidget(self.multiuser_session_inc_button)
+
+
         self.verticalLayout_8.addLayout(self.horizontalLayout_5)
+
+
         self.verticalLayout.addLayout(self.verticalLayout_8)
 
         self.horizontalLayout_2 = QHBoxLayout()
@@ -511,6 +530,8 @@ class Ui_switchboard_window(QMainWindow):
         self.horizontalLayout_2.setObjectName(u"horizontalLayout_2")
 
         self.verticalLayout.addLayout(self.horizontalLayout_2)
+
+
         self.horizontalLayout_4.addWidget(self.device_widget)
 
         self.frame_right = QFrame(self.centralwidget)
@@ -545,7 +566,7 @@ class Ui_switchboard_window(QMainWindow):
         self.logger_level_comboBox.addItem("")
         self.logger_level_comboBox.addItem("")
         self.logger_level_comboBox.setObjectName(u"logger_level_comboBox")
-        sizePolicy7 = QSizePolicy(QSizePolicy.Policy.Preferred, QSizePolicy.Policy.Fixed)
+        sizePolicy7 = QSizePolicy(QSizePolicy.Preferred, QSizePolicy.Fixed)
         sizePolicy7.setHorizontalStretch(1)
         sizePolicy7.setVerticalStretch(0)
         sizePolicy7.setHeightForWidth(self.logger_level_comboBox.sizePolicy().hasHeightForWidth())
@@ -553,7 +574,7 @@ class Ui_switchboard_window(QMainWindow):
 
         self.logger_options.addWidget(self.logger_level_comboBox)
 
-        self.logger_options_spacer_0 = QSpacerItem(40, 20, QSizePolicy.Policy.Expanding, QSizePolicy.Policy.Minimum)
+        self.logger_options_spacer_0 = QSpacerItem(40, 20, QSizePolicy.Expanding, QSizePolicy.Minimum)
 
         self.logger_options.addItem(self.logger_options_spacer_0)
 
@@ -567,7 +588,7 @@ class Ui_switchboard_window(QMainWindow):
 
         self.logger_options.addWidget(self.logger_wrap_label)
 
-        self.logger_options_spacer_1 = QSpacerItem(20, 20, QSizePolicy.Policy.Expanding, QSizePolicy.Policy.Minimum)
+        self.logger_options_spacer_1 = QSpacerItem(20, 20, QSizePolicy.Expanding, QSizePolicy.Minimum)
 
         self.logger_options.addItem(self.logger_options_spacer_1)
 
@@ -581,7 +602,7 @@ class Ui_switchboard_window(QMainWindow):
 
         self.logger_options.addWidget(self.logger_autoscroll_label)
 
-        self.logger_options_spacer_2 = QSpacerItem(20, 0, QSizePolicy.Policy.Expanding, QSizePolicy.Policy.Minimum)
+        self.logger_options_spacer_2 = QSpacerItem(20, 0, QSizePolicy.Expanding, QSizePolicy.Minimum)
 
         self.logger_options.addItem(self.logger_options_spacer_2)
 
@@ -590,6 +611,7 @@ class Ui_switchboard_window(QMainWindow):
 
         self.logger_options.addWidget(self.clear_log_button)
 
+
         self.verticalLayout_4.addLayout(self.logger_options)
 
         self.base_console = QPlainTextEdit(self.frame_logger)
@@ -655,23 +677,35 @@ class Ui_switchboard_window(QMainWindow):
 
         self.horizontalLayout_6 = QHBoxLayout()
         self.horizontalLayout_6.setObjectName(u"horizontalLayout_6")
-        self.horizontalSpacer_3 = QSpacerItem(40, 20, QSizePolicy.Policy.Expanding, QSizePolicy.Policy.Minimum)
+        self.horizontalSpacer_3 = QSpacerItem(40, 20, QSizePolicy.Expanding, QSizePolicy.Minimum)
 
         self.horizontalLayout_6.addItem(self.horizontalSpacer_3)
 
+        self.recording_status = QLabel(self.frame_logger)
+        self.recording_status.setObjectName(u"recording_status")
+
+        self.horizontalLayout_6.addWidget(self.recording_status)
+
+        self.line_6 = QFrame(self.frame_logger)
+        self.line_6.setObjectName(u"line_6")
+        self.line_6.setFrameShape(QFrame.VLine)
+        self.line_6.setFrameShadow(QFrame.Sunken)
+
+        self.horizontalLayout_6.addWidget(self.line_6)
+
         self.current_address_value = QLineEdit(self.frame_logger)
         self.current_address_value.setObjectName(u"current_address_value")
-        sizePolicy3.setHeightForWidth(self.current_address_value.sizePolicy().hasHeightForWidth())
-        self.current_address_value.setSizePolicy(sizePolicy3)
+        sizePolicy4.setHeightForWidth(self.current_address_value.sizePolicy().hasHeightForWidth())
+        self.current_address_value.setSizePolicy(sizePolicy4)
         self.current_address_value.setStyleSheet(u"border: 0px; background-color: rgba(0, 0, 0, 0)")
-        self.current_address_value.setAlignment(Qt.AlignRight | Qt.AlignTrailing | Qt.AlignVCenter)
+        self.current_address_value.setAlignment(Qt.AlignRight|Qt.AlignTrailing|Qt.AlignVCenter)
 
         self.horizontalLayout_6.addWidget(self.current_address_value)
 
         self.line_5 = QFrame(self.frame_logger)
         self.line_5.setObjectName(u"line_5")
-        self.line_5.setFrameShape(QFrame.Shape.VLine)
-        self.line_5.setFrameShadow(QFrame.Shadow.Sunken)
+        self.line_5.setFrameShape(QFrame.VLine)
+        self.line_5.setFrameShadow(QFrame.Sunken)
 
         self.horizontalLayout_6.addWidget(self.line_5)
 
@@ -679,27 +713,33 @@ class Ui_switchboard_window(QMainWindow):
         self.current_config_file_value.setObjectName(u"current_config_file_value")
 
         self.horizontalLayout_6.addWidget(self.current_config_file_value)
+
+
         self.verticalLayout_4.addLayout(self.horizontalLayout_6)
 
         self.splitter.addWidget(self.frame_logger)
 
         self.verticalLayout_2.addWidget(self.splitter)
+
+
         self.horizontalLayout_4.addWidget(self.frame_right)
 
         self.line_2 = QFrame(self.centralwidget)
         self.line_2.setObjectName(u"line_2")
-        self.line_2.setFrameShape(QFrame.Shape.VLine)
-        self.line_2.setFrameShadow(QFrame.Shadow.Sunken)
+        self.line_2.setFrameShape(QFrame.VLine)
+        self.line_2.setFrameShadow(QFrame.Sunken)
 
         self.horizontalLayout_4.addWidget(self.line_2)
+
+
         self.verticalLayout_3.addLayout(self.horizontalLayout_4)
 
-        self.setCentralWidget(self.centralwidget)
+        switchboard_window.setCentralWidget(self.centralwidget)
         self.line.raise_()
-        self.statusbar = QStatusBar(self)
+        self.statusbar = QStatusBar(switchboard_window)
         self.statusbar.setObjectName(u"statusbar")
-        self.setStatusBar(self.statusbar)
-        self.menu_bar = QMenuBar(self)
+        switchboard_window.setStatusBar(self.statusbar)
+        self.menu_bar = QMenuBar(switchboard_window)
         self.menu_bar.setObjectName(u"menu_bar")
         self.menu_bar.setGeometry(QRect(0, 0, 1192, 21))
         self.menu_configs = QMenu(self.menu_bar)
@@ -708,7 +748,7 @@ class Ui_switchboard_window(QMainWindow):
         self.menuConfigs.setObjectName(u"menuConfigs")
         self.menu_load_config = QMenu(self.menuConfigs)
         self.menu_load_config.setObjectName(u"menu_load_config")
-        self.setMenuBar(self.menu_bar)
+        switchboard_window.setMenuBar(self.menu_bar)
 
         self.menu_bar.addAction(self.menuConfigs.menuAction())
         self.menu_bar.addAction(self.menu_configs.menuAction())
@@ -719,12 +759,13 @@ class Ui_switchboard_window(QMainWindow):
         self.menuConfigs.addSeparator()
         self.menuConfigs.addAction(self.menu_load_config.menuAction())
 
-        self.retranslateUi(self)
+        self.retranslateUi(switchboard_window)
 
         self.tabs_main.setCurrentIndex(-1)
         self.logger_level_comboBox.setCurrentIndex(2)
 
-        QMetaObject.connectSlotsByName(self)
+
+        QMetaObject.connectSlotsByName(switchboard_window)
     # setupUi
 
     def retranslateUi(self, switchboard_window):
@@ -818,9 +859,11 @@ class Ui_switchboard_window(QMainWindow):
 #if QT_CONFIG(tooltip)
         self.clear_log_button.setToolTip(QCoreApplication.translate("switchboard_window", u"Clear the log window", None))
 #endif // QT_CONFIG(tooltip)
+        self.recording_status.setText(QCoreApplication.translate("switchboard_window", u"Recording Status", None))
         self.current_address_value.setText(QCoreApplication.translate("switchboard_window", u"Address", None))
         self.current_config_file_value.setText(QCoreApplication.translate("switchboard_window", u"TextLabel", None))
         self.menu_configs.setTitle(QCoreApplication.translate("switchboard_window", u"Settings", None))
         self.menuConfigs.setTitle(QCoreApplication.translate("switchboard_window", u"Configs", None))
         self.menu_load_config.setTitle(QCoreApplication.translate("switchboard_window", u"Load Config", None))
     # retranslateUi
+
